- LiveLink synchronization over ndisplay

diff --git a/Engine/Plugins/Runtime/LiveLinkOvernDisplay/LiveLinkOverNDisplay.uplugin b/Engine/Plugins/Runtime/LiveLinkOvernDisplay/LiveLinkOverNDisplay.uplugin
new file mode 100644
--- /dev/null
+++ b/Engine/Plugins/Runtime/LiveLinkOvernDisplay/LiveLinkOverNDisplay.uplugin
@@ -0,0 +1,38 @@
+{
+	"FileVersion": 3,
+	"Version": 1,
+	"VersionName": "1.0",
+	"FriendlyName": "LiveLinkOverNDisplay",
+	"CreatedBy": "Epic Games Inc",
+	"CreatedByURL": "http://epicgames.com",
+	"Description": "LiveLink subjects synchronization for nDisplay setup",
+	"Category": "Misc",
+	"DocsURL": "",
+	"MarketplaceURL": "",
+	"SupportURL": "",
+	"EnabledByDefault": false,
+	"CanContainContent": false,
+	"IsBetaVersion": true,
+	"Installed": false,
+	"SupportedTargetPlatforms": [ "Win64" ],
+	"Modules": [
+		{
+			"Name": "LiveLinkOverNDisplay",
+			"Type": "Runtime",
+			"LoadingPhase": "Default",
+			"WhitelistPlatforms": [
+				"Win64"
+			]
+		}
+	],
+	"Plugins": [
+		{
+			"Name": "nDisplay",
+			"Enabled": true
+		},
+		{
+			"Name": "LiveLink",
+			"Enabled": true
+		}
+	]
+}
diff --git a/Engine/Plugins/Runtime/LiveLinkOvernDisplay/Source/LiveLinkOverNDisplay/LiveLinkOverNDisplay.Build.cs b/Engine/Plugins/Runtime/LiveLinkOvernDisplay/Source/LiveLinkOverNDisplay/LiveLinkOverNDisplay.Build.cs
new file mode 100644
--- /dev/null
+++ b/Engine/Plugins/Runtime/LiveLinkOvernDisplay/Source/LiveLinkOverNDisplay/LiveLinkOverNDisplay.Build.cs
@@ -0,0 +1,35 @@
+// Copyright 1998-2019 Epic Games, Inc. All Rights Reserved.
+
+using UnrealBuildTool;
+using System.IO;
+
+public class LiveLinkOverNDisplay : ModuleRules
+{
+	public LiveLinkOverNDisplay(ReadOnlyTargetRules Target) : base(Target)
+	{
+		PublicDependencyModuleNames.AddRange(
+			new string[]
+			{
+				"Core",
+				"CoreUObject",
+			});
+
+		PrivateDependencyModuleNames.AddRange(
+			new string[]
+			{
+				"DisplayCluster",
+				"Engine",
+				"LiveLink",
+				"LiveLinkInterface",
+			});
+
+		if (Target.bBuildEditor == true)
+		{
+			PrivateDependencyModuleNames.AddRange(
+				new string[]
+				{
+					"Settings"
+				});
+		}
+	}
+}
diff --git a/Engine/Plugins/Runtime/LiveLinkOvernDisplay/Source/LiveLinkOverNDisplay/Private/LiveLinkOverNDisplayModule.cpp b/Engine/Plugins/Runtime/LiveLinkOvernDisplay/Source/LiveLinkOverNDisplay/Private/LiveLinkOverNDisplayModule.cpp
new file mode 100644
--- /dev/null
+++ b/Engine/Plugins/Runtime/LiveLinkOvernDisplay/Source/LiveLinkOverNDisplay/Private/LiveLinkOverNDisplayModule.cpp
@@ -0,0 +1,86 @@
+// Copyright 1998-2019 Epic Games, Inc. All Rights Reserved.
+
+#include "LiveLinkOverNDisplayModule.h"
+
+#include "DisplayClusterGameEngine.h"
+#include "Engine/Engine.h"
+#include "Features/IModularFeatures.h"
+#include "IDisplayCluster.h"
+#include "ILiveLinkClient.h"
+#include "LiveLinkOverNDisplayPrivate.h"
+#include "LiveLinkOverNDisplaySettings.h"
+#include "Misc/CoreDelegates.h"
+#include "Modules/ModuleManager.h"
+#include "NDisplayLiveLinkSubjectReplicator.h"
+
+#if WITH_EDITOR
+	#include "ISettingsModule.h"
+	#include "ISettingsSection.h"
+#endif 
+
+
+DEFINE_LOG_CATEGORY(LogLiveLinkOverNDisplay);
+
+#define LOCTEXT_NAMESPACE "LiveLinkOverNDisplayModule"
+
+
+FLiveLinkOverNDisplayModule::FLiveLinkOverNDisplayModule()
+	: LiveLinkReplicator(MakeUnique<FNDisplayLiveLinkSubjectReplicator>())
+{
+
+}
+
+void FLiveLinkOverNDisplayModule::StartupModule()
+{
+	// Register Engine callbacks
+	OnEngineLoopInitCompleteDelegate = FCoreDelegates::OnFEngineLoopInitComplete.AddRaw(this, &FLiveLinkOverNDisplayModule::OnEngineLoopInitComplete);
+
+#if WITH_EDITOR
+	if (ISettingsModule* SettingsModule = FModuleManager::GetModulePtr<ISettingsModule>("Settings"))
+	{
+		SettingsModule->RegisterSettings(
+			"Project", "Plugins", "LiveLink over nDisplay",
+			LOCTEXT("LiveLinkOverNDisplaySettingsName", "LiveLink over nDisplay"),
+			LOCTEXT("LiveLinkOverNDisplaySettingsDescription", "Configure LiveLink over nDisplay."),
+			GetMutableDefault<ULiveLinkOverNDisplaySettings>()
+		);
+	}
+#endif
+}
+
+void FLiveLinkOverNDisplayModule::ShutdownModule()
+{
+#if WITH_EDITOR
+	if (ISettingsModule* SettingsModule = FModuleManager::GetModulePtr<ISettingsModule>("Settings"))
+	{
+		SettingsModule->UnregisterSettings("Project", "Plugins", "LiveLink over nDisplay");
+	}
+#endif
+
+	FCoreDelegates::OnFEngineLoopInitComplete.RemoveAll(this);
+}
+
+FNDisplayLiveLinkSubjectReplicator& FLiveLinkOverNDisplayModule::GetSubjectReplicator()
+{
+	return *LiveLinkReplicator;
+}
+
+void FLiveLinkOverNDisplayModule::OnEngineLoopInitComplete()
+{
+	if (UDisplayClusterGameEngine* DisplayClusterEngine = Cast<UDisplayClusterGameEngine>(GEngine))
+	{
+		if (EDisplayClusterOperationMode::Cluster == DisplayClusterEngine->GetOperationMode())
+		{
+			if (GetDefault<ULiveLinkOverNDisplaySettings>()->IsLiveLinkOverNDisplayEnabled())
+			{
+				LiveLinkReplicator->Activate();
+			}
+		}
+	}
+}
+
+#undef LOCTEXT_NAMESPACE
+
+
+IMPLEMENT_MODULE(FLiveLinkOverNDisplayModule, LiveLinkOverNDisplay);
+
diff --git a/Engine/Plugins/Runtime/LiveLinkOvernDisplay/Source/LiveLinkOverNDisplay/Private/LiveLinkOverNDisplayModule.h b/Engine/Plugins/Runtime/LiveLinkOvernDisplay/Source/LiveLinkOverNDisplay/Private/LiveLinkOverNDisplayModule.h
new file mode 100644
--- /dev/null
+++ b/Engine/Plugins/Runtime/LiveLinkOvernDisplay/Source/LiveLinkOverNDisplay/Private/LiveLinkOverNDisplayModule.h
@@ -0,0 +1,43 @@
+// Copyright 1998-2019 Epic Games, Inc. All Rights Reserved.
+
+#pragma once
+
+
+#include "ILiveLinkOverNDisplayModule.h"
+#include "LiveLinkVirtualSubject.h"
+
+
+class FNDisplayLiveLinkSubjectReplicator;
+
+
+/**
+ * Entry point for LiveLinkOverNDisplay functionnality. 
+ */
+class FLiveLinkOverNDisplayModule : public ILiveLinkOverNDisplayModule, public TSharedFromThis<FLiveLinkOverNDisplayModule>
+{
+public:
+
+	FLiveLinkOverNDisplayModule();
+
+	//~ Begin IModuleInterface interface
+	virtual void StartupModule() override;
+	virtual void ShutdownModule() override;
+	//~ End IModuleInterface interface
+
+	virtual FNDisplayLiveLinkSubjectReplicator& GetSubjectReplicator() override;
+
+private:
+
+	void OnEngineLoopInitComplete();
+
+private:
+
+	/** Cluster event listener delegate */
+	FDelegateHandle OnEngineLoopInitCompleteDelegate;
+
+	/** Console commands handle. */
+	TUniquePtr<FAutoConsoleCommand> StopMessageSyncCommand;
+
+	/** Replicator SyncObject used to transfer data across all cluster machines */
+	TUniquePtr<FNDisplayLiveLinkSubjectReplicator> LiveLinkReplicator;
+};
diff --git a/Engine/Plugins/Runtime/LiveLinkOvernDisplay/Source/LiveLinkOverNDisplay/Private/LiveLinkOverNDisplayPrivate.h b/Engine/Plugins/Runtime/LiveLinkOvernDisplay/Source/LiveLinkOverNDisplay/Private/LiveLinkOverNDisplayPrivate.h
new file mode 100644
--- /dev/null
+++ b/Engine/Plugins/Runtime/LiveLinkOvernDisplay/Source/LiveLinkOverNDisplay/Private/LiveLinkOverNDisplayPrivate.h
@@ -0,0 +1,7 @@
+// Copyright 1998-2019 Epic Games, Inc. All Rights Reserved.
+
+#pragma once
+
+#include "CoreMinimal.h"
+
+DECLARE_LOG_CATEGORY_EXTERN(LogLiveLinkOverNDisplay, Log, All);
diff --git a/Engine/Plugins/Runtime/LiveLinkOvernDisplay/Source/LiveLinkOverNDisplay/Private/LiveLinkOverNDisplaySettings.cpp b/Engine/Plugins/Runtime/LiveLinkOvernDisplay/Source/LiveLinkOverNDisplay/Private/LiveLinkOverNDisplaySettings.cpp
new file mode 100644
--- /dev/null
+++ b/Engine/Plugins/Runtime/LiveLinkOvernDisplay/Source/LiveLinkOverNDisplay/Private/LiveLinkOverNDisplaySettings.cpp
@@ -0,0 +1,18 @@
+// Copyright 1998-2019 Epic Games, Inc. All Rights Reserved.
+
+#include "LiveLinkOverNDisplaySettings.h"
+
+#include "LiveLinkOverNDisplayPrivate.h"
+
+
+ULiveLinkOverNDisplaySettings::ULiveLinkOverNDisplaySettings()
+{
+	FString BoolValue;
+	const bool bIsFoundOnCommandLine = FParse::Value(FCommandLine::Get(), TEXT("-EnableLiveLinkOverNDisplay="), BoolValue);
+
+	if (bIsFoundOnCommandLine)
+	{
+		UE_LOG(LogLiveLinkOverNDisplay, Log, TEXT("Overriding LiveLinkOverNDisplay enable flag from command line with value '%s'"), *BoolValue);
+		bIsEnabledFromCommandLine = BoolValue.ToBool();
+	}
+}
diff --git a/Engine/Plugins/Runtime/LiveLinkOvernDisplay/Source/LiveLinkOverNDisplay/Private/NDisplayLiveLinkSubjectReplicator.cpp b/Engine/Plugins/Runtime/LiveLinkOvernDisplay/Source/LiveLinkOverNDisplay/Private/NDisplayLiveLinkSubjectReplicator.cpp
new file mode 100644
--- /dev/null
+++ b/Engine/Plugins/Runtime/LiveLinkOvernDisplay/Source/LiveLinkOverNDisplay/Private/NDisplayLiveLinkSubjectReplicator.cpp
@@ -0,0 +1,419 @@
+// Copyright 1998-2019 Epic Games, Inc. All Rights Reserved.
+
+#include "NDisplayLiveLinkSubjectReplicator.h"
+
+
+#include "Cluster/IDisplayClusterClusterManager.h"
+#include "Features/IModularFeatures.h"
+#include "IDisplayCluster.h"
+#include "ILiveLinkClient.h"
+#include "LiveLinkSubjectSettings.h"
+#include "Misc/CoreDelegates.h"
+#include "NDisplaySlaveVirtualSubject.h"
+
+
+DEFINE_LOG_CATEGORY_STATIC(LogNDisplayLiveLinkSubjectReplicator, Log, All);
+
+
+
+FNDisplayLiveLinkSubjectReplicator::~FNDisplayLiveLinkSubjectReplicator()
+{
+	if (IDisplayCluster::IsAvailable())
+	{
+		if (IDisplayClusterClusterManager* ClusterManager = IDisplayCluster::Get().GetClusterMgr())
+		{
+			ClusterManager->UnregisterSyncObject(this);
+		}
+	}
+	
+	if (LiveLinkClient)
+	{
+		LiveLinkClient->OnLiveLinkTicked().RemoveAll(this);
+	}
+
+	FCoreDelegates::OnFEngineLoopInitComplete.RemoveAll(this);
+	FCoreDelegates::OnBeginFrame.RemoveAll(this);
+}
+
+bool FNDisplayLiveLinkSubjectReplicator::IsActive() const
+{
+	return LiveLinkClient != nullptr;
+}
+
+FString FNDisplayLiveLinkSubjectReplicator::GetSyncId() const
+{
+	static const FString SyncId = TEXT("NDisplayLiveLinkSyncObject");
+	return SyncId;
+}
+
+FString FNDisplayLiveLinkSubjectReplicator::SerializeToString() const
+{
+	//Serialization is done on the master. Payload was created on LiveLinkTicked callback
+	FScopeLock Lock(&PayloadCriticalSection);
+	const FString ThisFramePayload = LiveLinkPayload;
+	return ThisFramePayload;
+}
+
+bool FNDisplayLiveLinkSubjectReplicator::DeserializeFromString(const FString& Str)
+{
+	//Deserialization is done on each slave
+	TArray<uint8> TempBytes;
+	TempBytes.AddUninitialized(Str.Len());
+	StringToBytes(Str, TempBytes.GetData(), Str.Len());
+
+	FMemoryReader Ar(TempBytes);
+	OnDataSynchronization(Ar);
+
+	return true;
+}
+
+void FNDisplayLiveLinkSubjectReplicator::AddReferencedObjects(FReferenceCollector& Collector)
+{
+	for (const UNDisplaySlaveVirtualSubject* Subject : TrackedSubjects)
+	{
+		Collector.AddReferencedObject(Subject);
+	}
+}
+
+void FNDisplayLiveLinkSubjectReplicator::Activate()
+{
+	if (IModularFeatures::Get().IsModularFeatureAvailable(ILiveLinkClient::ModularFeatureName))
+	{
+		LiveLinkClient = &IModularFeatures::Get().GetModularFeature<ILiveLinkClient>(ILiveLinkClient::ModularFeatureName);
+
+		if (IDisplayCluster::Get().GetClusterMgr()->IsMaster())
+		{
+			//If we're a master, listen to LiveLinkTicked callback to bundle up the SyncObject. Subjects frames will have been processed at that point.
+			LiveLinkClient->OnLiveLinkTicked().AddRaw(this, &FNDisplayLiveLinkSubjectReplicator::OnLiveLinkTicked);
+		}
+		else if (IDisplayCluster::Get().GetClusterMgr()->IsSlave())
+		{
+			//If we're a slave, listen for new subject and disables them if we're tracking that subject
+			FCoreDelegates::OnBeginFrame.AddRaw(this, &FNDisplayLiveLinkSubjectReplicator::OnEngineBeginFrame);
+
+			//Create our Source to hold all VirtualSubject we'll create
+			LiveLinkSourceGuid = LiveLinkClient->AddVirtualSubjectSource(TEXT("nDisplaySubjectReplicator"));
+			if (!LiveLinkSourceGuid.IsValid())
+			{
+				UE_LOG(LogNDisplayLiveLinkSubjectReplicator, Error, TEXT("Could not add VirtualSubject Source for nDisplay replication."));
+			}
+		}
+
+		//Register SyncObject on PreTick to have this behavior
+
+		//Order of operation will be
+		//1. When LiveLink creates snapshots for the frame, Master will create the sync object payload
+		//2. On PreTick, the SyncObject will be synchronized between cluster machines
+		//3. On object processing, Slaves will push received data for each subject to have the same data then Master
+		IDisplayClusterClusterManager* ClusterManager = IDisplayCluster::Get().GetClusterMgr();
+		ClusterManager->RegisterSyncObject(this, EDisplayClusterSyncGroup::PreTick);
+	}
+	else
+	{
+		UE_LOG(LogNDisplayLiveLinkSubjectReplicator, Error, TEXT("Can't activate LiveLink Subject Replicator for nDisplay because LiveLink is not available."));
+	}
+}
+
+void FNDisplayLiveLinkSubjectReplicator::OnLiveLinkTicked()
+{
+	TArray<uint8> TempBytes;
+	FMemoryWriter Ar(TempBytes, /*bIsPersistent=*/ true);
+	OnDataSynchronization(Ar);
+
+	FScopeLock Lock(&PayloadCriticalSection);
+	LiveLinkPayload = BytesToString(TempBytes.GetData(), TempBytes.Num());
+}
+
+void FNDisplayLiveLinkSubjectReplicator::OnEngineBeginFrame()
+{
+	//Each frame, disable all subjects. When we are enabled, we are replicating ALL subjects from the master to each slave.
+	//Only our replicated VirtualSubjects should be enabled
+	if (LiveLinkClient)
+	{
+		check(IDisplayCluster::Get().GetClusterMgr()->IsSlave());
+
+		const bool bIncludeDisabled = false;
+		const bool bIncludeVirtual = true;
+		const TArray<FLiveLinkSubjectKey> Subjects = LiveLinkClient->GetSubjects(bIncludeDisabled, bIncludeVirtual);
+		for (const FLiveLinkSubjectKey& Subject : Subjects)
+		{
+			if (Subject.Source != LiveLinkSourceGuid)
+			{
+				LiveLinkClient->SetSubjectEnabled(Subject, false);
+			}
+		}
+		
+		//Make sure all our replicated subjects are enabled
+		for (const UNDisplaySlaveVirtualSubject* OurSubject : TrackedSubjects)
+		{
+			LiveLinkClient->SetSubjectEnabled(OurSubject->GetSubjectKey(), true);
+		}
+	}
+}
+
+void FNDisplayLiveLinkSubjectReplicator::OnDataSynchronization(FArchive& Ar)
+{
+	//Master is saving, slaves are loading
+	const bool bIsSaving = Ar.IsSaving();
+	bool bContinue = bIsSaving;
+
+	EFrameType FrameType = EFrameType::DataOnly;
+
+	//Fill in current tracked subject list. If a subject is not processed this frame, we remove it
+	TArray<FLiveLinkSubjectKey, TInlineAllocator<16>> LastFrameSubjectKeys;
+	for (const UNDisplaySlaveVirtualSubject* Subject : TrackedSubjects)
+	{
+		LastFrameSubjectKeys.Add(Subject->GetAssociatedSubjectKey());
+	}
+
+	if (bIsSaving)
+	{
+		if (LiveLinkClient)
+		{
+			/**
+			 * We're the master, and serializing data out to clients:
+			 *	- Iterate over all of the snapshots.
+			 *		- Determine the type of frame this is (new static data, new subject, data only).
+			 *		- Serialize a True continuation byte.
+			 *		- Serialize the frame.
+			 *	- Once all snapshots have been added, serialize a False continuation byte.
+			 */
+			{
+				//Get all subjects except the ones that are disabled
+				//Since we are replicating using a VirtualSubject, slaves will have to remove any VirtualSubject being replicated
+				const bool bIncludeDisabled = false;
+				const bool bIncludeVirtuals = true;
+				TArray<FLiveLinkSubjectKey> ThisFramesSubjects = LiveLinkClient->GetSubjects(bIncludeDisabled, bIncludeVirtuals);
+				for (FLiveLinkSubjectKey& SubjectKey : ThisFramesSubjects)
+				{
+					const FName SubjectName = SubjectKey.SubjectName;
+					TSubclassOf<ULiveLinkRole> SubjectRole = LiveLinkClient->GetSubjectRole(SubjectKey);
+
+					FLiveLinkSubjectFrameData SubjectFrameData;
+					if (LiveLinkClient->EvaluateFrame_AnyThread(SubjectName, SubjectRole, SubjectFrameData))
+					{
+						if (UNDisplaySlaveVirtualSubject** FoundSujectDataPtr = TrackedSubjects.FindByPredicate(
+							[SubjectKey](const UNDisplaySlaveVirtualSubject* Other) { return Other->GetAssociatedSubjectKey().SubjectName == SubjectKey.SubjectName; }))
+						{
+							UNDisplaySlaveVirtualSubject* FoundSubjectData = *FoundSujectDataPtr;
+							if (SubjectFrameData.StaticData != FoundSubjectData->GetStaticData() 
+								|| FoundSubjectData->GetAssociatedSubjectKey() != SubjectKey 
+								|| FoundSubjectData->GetRole() != SubjectRole)
+							{
+								FrameType = EFrameType::UpdatedSubject;
+							}
+							else
+							{
+								FrameType = EFrameType::DataOnly;
+							}
+						}
+						else
+						{
+							FrameType = EFrameType::NewSubject;
+						}
+
+						//This subject was processed
+						LastFrameSubjectKeys.Remove(SubjectKey);
+
+						Ar << bContinue;
+
+						HandleFrame(Ar, FrameType, SubjectKey, SubjectRole, SubjectFrameData);
+					}
+					else
+					{
+						UE_LOG(LogNDisplayLiveLinkSubjectReplicator, Verbose, TEXT("Master could not evaluate Subject '%s'. It could have starved."), *SubjectName.ToString());
+					}
+				}
+			}
+		}
+
+		//We're done processing subjects, mark the frame as done 
+		bContinue = false;
+		Ar << bContinue;
+	}
+	else
+	{
+		/**
+		 * We're a slave, and serializing data from the master:
+		 *	- Serialize the continuation byte
+		 *		- If False, we're done.
+		 *		- If True:
+		 *			- Serialize the frame.
+		 *			- Based on frame type, do required process.
+		 *			- Repeat.
+		 */
+
+		TSubclassOf<ULiveLinkRole> SubjectRole;
+		FLiveLinkSubjectKey SubjectKey;
+		FLiveLinkSubjectFrameData SubjectFrameData;
+
+		while (true)
+		{
+			Ar << bContinue;
+			if (bContinue)
+			{
+				//Process frame data and fill this subject's data
+				HandleFrame(Ar, FrameType, SubjectKey, SubjectRole, SubjectFrameData);
+				if (!SubjectKey.SubjectName.IsNone())
+				{
+					ProcessLiveLinkData_Slave(FrameType, SubjectKey, SubjectFrameData.FrameData);
+
+					LastFrameSubjectKeys.RemoveSingleSwap(SubjectKey);
+				}
+			}
+			else
+			{
+				break;
+			}
+		}
+	}
+
+	// Cleanup missing subjects. Slaves will need to remove VirtualSubjects that were created
+	if (LiveLinkClient)
+	{
+		for (const FLiveLinkSubjectKey& ItemSubjectKey : LastFrameSubjectKeys)
+		{
+			UE_LOG(LogNDisplayLiveLinkSubjectReplicator, Verbose, TEXT("Subject '%s' was not processed. Removing it from our list."), *ItemSubjectKey.SubjectName.ToString());
+
+			if (!bIsSaving)
+			{
+				const FLiveLinkSubjectKey ReplicatedSubjectKey(LiveLinkSourceGuid, ItemSubjectKey.SubjectName);
+				LiveLinkClient->RemoveVirtualSubject(ReplicatedSubjectKey);
+			}
+
+			const int32 FoundIndex = TrackedSubjects.IndexOfByPredicate([ItemSubjectKey](const UNDisplaySlaveVirtualSubject* Other) { return Other->GetAssociatedSubjectKey() == ItemSubjectKey; });
+			if (TrackedSubjects.IsValidIndex(FoundIndex))
+			{
+				TrackedSubjects.RemoveAtSwap(FoundIndex);
+			}
+		}
+	}
+	else
+	{
+		TrackedSubjects.Reset();
+	}
+}
+
+void FNDisplayLiveLinkSubjectReplicator::HandleFrame(FArchive& Ar, EFrameType& FrameType, FLiveLinkSubjectKey& SubjectKey, TSubclassOf<ULiveLinkRole>& SubjectRole, FLiveLinkSubjectFrameData& SubjectFrame)
+{
+	//Archive always start with the type and subject identifier
+	Ar << FrameType;
+	Ar << SubjectKey;
+
+	//If it's a new subject / updated one, process static data and role
+	if (FrameType == EFrameType::NewSubject || FrameType == EFrameType::UpdatedSubject)
+	{
+		HandleNewSubject(Ar, SubjectKey, SubjectRole, SubjectFrame);
+	}
+
+	//End archive with Frame data
+	Ar << SubjectFrame.FrameData;
+	
+	UE_LOG(LogNDisplayLiveLinkSubjectReplicator, VeryVerbose, TEXT("HandleFrame SubjectName=%s FrameType=%d"), *SubjectKey.SubjectName.ToString(), static_cast<uint8>(FrameType));
+}
+
+void FNDisplayLiveLinkSubjectReplicator::HandleNewSubject(FArchive& Ar, FLiveLinkSubjectKey& SubjectKey, TSubclassOf<ULiveLinkRole>& SubjectRole, FLiveLinkSubjectFrameData& SubjectFrame)
+{
+	//New or updated subjects will have static data 
+	Ar << SubjectFrame.StaticData;
+
+	UE_LOG(LogNDisplayLiveLinkSubjectReplicator, Verbose, TEXT("HandleNewSubject SubjectName=%s"), *SubjectKey.SubjectName.ToString());
+
+	//Verify if we're already tracking that subject
+	UNDisplaySlaveVirtualSubject* TrackedSubject = nullptr;
+	if (UNDisplaySlaveVirtualSubject** FoundSujectDataPtr = TrackedSubjects.FindByPredicate(
+		[SubjectKey](const UNDisplaySlaveVirtualSubject* Other) { return Other->GetAssociatedSubjectKey().SubjectName == SubjectKey.SubjectName; }))
+	{
+		TrackedSubject = *FoundSujectDataPtr;
+	}
+
+	if (Ar.IsSaving())
+	{
+		FString RoleClassName = SubjectRole->GetPathName();
+		Ar << RoleClassName;
+
+		//If we're not tracking that subject yet, create it
+		if (TrackedSubject == nullptr)
+		{
+			TrackedSubject = NewObject<UNDisplaySlaveVirtualSubject>(GetTransientPackage());
+			TrackedSubjects.Add(TrackedSubject);
+		}
+		
+		//When handling new subject, always setup role and static data
+		TrackedSubject->SetTrackedSubjectInfo(SubjectKey, SubjectRole);
+		TrackedSubject->GetStaticData().InitializeWith(SubjectFrame.StaticData);
+	}
+	else
+	{
+		FString RoleClassName;
+		Ar << RoleClassName;
+		SubjectRole = StaticLoadClass(UObject::StaticClass(), NULL, *RoleClassName, NULL, LOAD_None, NULL);
+		check(SubjectRole);
+
+		//On Slaves, create the virtual subject directly in LiveLink so it's part of the system
+		ULiveLinkVirtualSubject* ReplicatedSubject = Cast<ULiveLinkVirtualSubject>(LiveLinkClient->GetSubjectSettings(SubjectKey));
+
+		//Always disable the subject we're replicating
+		LiveLinkClient->SetSubjectEnabled(SubjectKey, false);
+
+		//If we're not tracking that subject yet, create it
+		if (TrackedSubject == nullptr)
+		{
+			const FLiveLinkSubjectKey ReplicatedSubjectKey(LiveLinkSourceGuid, SubjectKey.SubjectName);
+			LiveLinkClient->AddVirtualSubject(ReplicatedSubjectKey, UNDisplaySlaveVirtualSubject::StaticClass());
+
+			//Retrieve the newly created subject to add it to our tracking array
+			const bool bIncludeDisabled = false;
+			const bool bIncludeVirtual = true;
+			for (const FLiveLinkSubjectKey& FoundSubjectKey : LiveLinkClient->GetSubjects(bIncludeDisabled, bIncludeVirtual))
+			{
+				if (FoundSubjectKey.SubjectName == SubjectKey.SubjectName)
+				{
+					//If this subject has the same subject name, verify if it's the VirtualOne. Settings object will return the VirtualSubject directly
+					if (UNDisplaySlaveVirtualSubject* NewlyCreatedSubject = Cast<UNDisplaySlaveVirtualSubject>(LiveLinkClient->GetSubjectSettings(FoundSubjectKey)))
+					{
+						TrackedSubject = NewlyCreatedSubject;
+						TrackedSubjects.Add(NewlyCreatedSubject);
+					}
+				}
+			}
+		}
+
+		//VirtualSubjects need translators to be able to be evaluated in different output format. Special treatment required for virtual subjects since they won't/can't exist in the client
+		if (ReplicatedSubject)
+		{
+			const TArray<ULiveLinkFrameTranslator*>& Translators = ReplicatedSubject->GetTranslators();
+			TrackedSubject->UpdateTranslators(Translators);
+		}
+		else
+		{
+			if (ULiveLinkSubjectSettings* SubjectSetting = Cast<ULiveLinkSubjectSettings>(LiveLinkClient->GetSubjectSettings(SubjectKey)))
+			{
+				TrackedSubject->UpdateTranslators(SubjectSetting->Translators);
+			}
+			else
+			{
+				UE_LOG(LogNDisplayLiveLinkSubjectReplicator, Warning, TEXT("Replicating subject '%s' but could not find its settings"), *SubjectKey.SubjectName.ToString());
+			}
+		}
+
+		//When handling new subject, always setup role and static data
+		TrackedSubject->SetTrackedSubjectInfo(SubjectKey, SubjectRole);
+		TrackedSubject->GetStaticData().InitializeWith(SubjectFrame.StaticData);
+	}
+}
+
+void FNDisplayLiveLinkSubjectReplicator::ProcessLiveLinkData_Slave(EFrameType FrameType, const FLiveLinkSubjectKey& SubjectKey, FLiveLinkFrameDataStruct& FrameData)
+{
+	//Find associated tracked subject data to broadcast role and static data. Should always be found
+	UNDisplaySlaveVirtualSubject** FoundTrackedSubjectPtr = TrackedSubjects.FindByPredicate([SubjectKey](const UNDisplaySlaveVirtualSubject* Other) { return Other->GetAssociatedSubjectKey() == SubjectKey; });
+	check(FoundTrackedSubjectPtr);
+	UNDisplaySlaveVirtualSubject* FoundTrackedSubject = *FoundTrackedSubjectPtr;
+	
+	FrameData.GetBaseData()->WorldTime = FPlatformTime::Seconds();
+	FoundTrackedSubject->UpdateFrameData(MoveTemp(FrameData));
+
+	UE_LOG(LogNDisplayLiveLinkSubjectReplicator, VeryVerbose, TEXT("ProcessLiveLinkData SubjectUpdated=%s UpdateType=%d Timecode=%s")
+		, *(SubjectKey.SubjectName.ToString()), static_cast<uint8>(FrameType)
+		, *(FTimecode::FromFrameNumber(FrameData.GetBaseData()->MetaData.SceneTime.Time.FrameNumber, FrameData.GetBaseData()->MetaData.SceneTime.Rate, false).ToString()));
+}
diff --git a/Engine/Plugins/Runtime/LiveLinkOvernDisplay/Source/LiveLinkOverNDisplay/Private/NDisplayLiveLinkSubjectReplicator.h b/Engine/Plugins/Runtime/LiveLinkOvernDisplay/Source/LiveLinkOverNDisplay/Private/NDisplayLiveLinkSubjectReplicator.h
new file mode 100644
--- /dev/null
+++ b/Engine/Plugins/Runtime/LiveLinkOvernDisplay/Source/LiveLinkOverNDisplay/Private/NDisplayLiveLinkSubjectReplicator.h
@@ -0,0 +1,97 @@
+// Copyright 1998-2019 Epic Games, Inc. All Rights Reserved.
+
+#pragma once
+
+#include "Cluster/IDisplayClusterClusterSyncObject.h"
+#include "UObject/GCObject.h"
+
+
+#include "LiveLinkClientReference.h"
+
+
+
+class UNDisplaySlaveVirtualSubject;
+
+
+/**
+ * Classes used to replicates data to be used for each frame for each enabled subjects on the Master machines.
+ * Slaves will use that replicated data to use the same thing on each machines
+ */
+class LIVELINKOVERNDISPLAY_API FNDisplayLiveLinkSubjectReplicator : public IDisplayClusterClusterSyncObject, public FGCObject
+{
+private:
+	enum class EFrameType : uint8
+	{
+		DataOnly,
+		NewSubject,			//A new subject was sent this frame.
+		UpdatedSubject,		//The subject's static data or role was updated this frame.
+	};
+
+public:
+
+	FNDisplayLiveLinkSubjectReplicator() = default;
+	virtual ~FNDisplayLiveLinkSubjectReplicator();
+
+	/** Move only type since we're owning LiveLink base structs */
+	FNDisplayLiveLinkSubjectReplicator(const FNDisplayLiveLinkSubjectReplicator&) = delete;
+	FNDisplayLiveLinkSubjectReplicator& operator=(const FNDisplayLiveLinkSubjectReplicator&) = delete;
+	FNDisplayLiveLinkSubjectReplicator(FNDisplayLiveLinkSubjectReplicator&&) = default;
+
+	//~ Begin IDisplayClusterClusterSyncObject interface
+	virtual bool IsActive() const override;
+	virtual FString GetSyncId() const override;
+	virtual bool IsDirty() const override { return true; };
+	virtual void ClearDirty() override {};
+	virtual FString SerializeToString() const override;
+	virtual bool DeserializeFromString(const FString& Ar) override;
+	//~ End IDisplayClusterClusterSyncObject interface
+
+	//~ Begin FGCObject interface
+	virtual void AddReferencedObjects(FReferenceCollector& Collector) override;
+	//~ End FGCObject interface
+
+	void Activate();
+
+protected:
+
+	/**
+     * Callback to be notified when LiveLinkClient was ticked so SyncObject can be serialized
+     * @note Only called on Master
+     */
+	void OnLiveLinkTicked();
+
+	/**
+	 * Hook to beginning of frames to make sure our replicated virtual subjects are the ones enabled
+	 * @note Only called on Slave
+	 */
+	void OnEngineBeginFrame();
+
+	/** Synchronization point when object is going to be synchronized across nDisplay cluster */
+	void OnDataSynchronization(FArchive& Ar);
+
+	/** Process new subject data when serializing SyncObject */
+	void HandleNewSubject(FArchive& Ar, FLiveLinkSubjectKey& SubjectKey, TSubclassOf<ULiveLinkRole>& SubjectRole, FLiveLinkSubjectFrameData& SubjectFrame);
+
+	/** Process a frame for a specific subject */
+	void HandleFrame(FArchive& Ar, EFrameType& FrameType, FLiveLinkSubjectKey& SubjectKey, TSubclassOf<ULiveLinkRole>& SubjectRole, FLiveLinkSubjectFrameData& SubjectFrame);
+
+	/** Only used on slaves, handles this frame data for a subject. If a new subject is handled, the associated VirtualSubject will be created */
+	void ProcessLiveLinkData_Slave(EFrameType FrameType, const FLiveLinkSubjectKey& SubjectKey, FLiveLinkFrameDataStruct& FrameData);
+
+private:
+
+	/** Cached LiveLinkClient when modular feature is registered */
+	ILiveLinkClient* LiveLinkClient = nullptr;
+
+	//CriticalSection to protect the payload. SerializeToString can bel called from other threads than Main Thread
+	mutable FCriticalSection PayloadCriticalSection;
+
+	/** SyncObject uses string to pass around data. This represents binary LiveLinkSubject information converted to string */
+	FString LiveLinkPayload;
+
+	/** List of Subjects that we are replicating across cluster. On Slaves, it will actually be added to the LiveLink subject's list */
+	TArray<UNDisplaySlaveVirtualSubject*> TrackedSubjects;
+
+	/** Guid associated to our Virtual Subject Source */
+	FGuid LiveLinkSourceGuid;
+};
diff --git a/Engine/Plugins/Runtime/LiveLinkOvernDisplay/Source/LiveLinkOverNDisplay/Private/NDisplaySlaveVirtualSubject.cpp b/Engine/Plugins/Runtime/LiveLinkOvernDisplay/Source/LiveLinkOverNDisplay/Private/NDisplaySlaveVirtualSubject.cpp
new file mode 100644
--- /dev/null
+++ b/Engine/Plugins/Runtime/LiveLinkOvernDisplay/Source/LiveLinkOverNDisplay/Private/NDisplaySlaveVirtualSubject.cpp
@@ -0,0 +1,33 @@
+// Copyright 1998-2019 Epic Games, Inc. All Rights Reserved.
+
+#include "NDisplaySlaveVirtualSubject.h"
+
+#include "LiveLinkOverNDisplayPrivate.h"
+#include "LiveLinkSubjectSettings.h"
+
+
+
+void UNDisplaySlaveVirtualSubject::Update()
+{
+	UpdateTranslatorsForThisFrame();
+}
+
+void UNDisplaySlaveVirtualSubject::UpdateFrameData(FLiveLinkFrameDataStruct&& NewFrameData)
+{
+	FrameSnapshot.FrameData = MoveTemp(NewFrameData);
+}
+
+void UNDisplaySlaveVirtualSubject::SetTrackedSubjectInfo(const FLiveLinkSubjectKey& InSubjectKey, TSubclassOf<ULiveLinkRole> InRole)
+{
+	AssociatedSubject = InSubjectKey;
+	Role = InRole;
+}
+
+void UNDisplaySlaveVirtualSubject::UpdateTranslators(const TArray<ULiveLinkFrameTranslator*>& SourceTranslators)
+{
+	for (const ULiveLinkFrameTranslator* SourceTranslator : SourceTranslators)
+	{
+		ULiveLinkFrameTranslator* NewTranslator = DuplicateObject(SourceTranslator, this);
+		FrameTranslators.Add(NewTranslator);
+	}
+}
diff --git a/Engine/Plugins/Runtime/LiveLinkOvernDisplay/Source/LiveLinkOverNDisplay/Private/NDisplaySlaveVirtualSubject.h b/Engine/Plugins/Runtime/LiveLinkOvernDisplay/Source/LiveLinkOverNDisplay/Private/NDisplaySlaveVirtualSubject.h
new file mode 100644
--- /dev/null
+++ b/Engine/Plugins/Runtime/LiveLinkOvernDisplay/Source/LiveLinkOverNDisplay/Private/NDisplaySlaveVirtualSubject.h
@@ -0,0 +1,41 @@
+// Copyright 1998-2019 Epic Games, Inc. All Rights Reserved.
+
+#pragma once
+
+#include "LiveLinkVirtualSubject.h"
+
+#include "Cluster/IDisplayClusterClusterSyncObject.h"
+
+#include "NDisplaySlaveVirtualSubject.generated.h"
+
+
+/**
+ * LiveLink VirtualSubject used on nDisplay slaves to replicate real subjects
+ * Master sends data to use for each frame to make sure all machines display / uses the same data
+ * Uses Replicator object that this module creates.
+ */
+UCLASS()
+class UNDisplaySlaveVirtualSubject : public ULiveLinkVirtualSubject
+{
+	GENERATED_BODY()
+
+public:
+	
+	//~ Begin ULiveLinkVirtualSubject interface
+	virtual void Update() override;
+	//~ End ULiveLinkVirtualSubject interface
+
+	void UpdateFrameData(FLiveLinkFrameDataStruct&& NewFrameData);
+
+
+	void SetTrackedSubjectInfo(const FLiveLinkSubjectKey& SubjectKey, TSubclassOf<ULiveLinkRole> Role);
+
+	void UpdateTranslators(const TArray<ULiveLinkFrameTranslator*>& SourceTranslators);
+
+	FLiveLinkSubjectKey GetAssociatedSubjectKey() const { return AssociatedSubject; }
+
+private:
+
+	/** Real subject that this virtual subject replicates */
+	FLiveLinkSubjectKey AssociatedSubject;
+};
diff --git a/Engine/Plugins/Runtime/LiveLinkOvernDisplay/Source/LiveLinkOverNDisplay/Public/ILiveLinkOverNDisplayModule.h b/Engine/Plugins/Runtime/LiveLinkOvernDisplay/Source/LiveLinkOverNDisplay/Public/ILiveLinkOverNDisplayModule.h
new file mode 100644
--- /dev/null
+++ b/Engine/Plugins/Runtime/LiveLinkOvernDisplay/Source/LiveLinkOverNDisplay/Public/ILiveLinkOverNDisplayModule.h
@@ -0,0 +1,37 @@
+// Copyright 1998-2019 Epic Games, Inc. All Rights Reserved.
+
+#pragma once
+
+#include "Modules/ModuleInterface.h"
+#include "Modules/ModuleManager.h"
+
+class FNDisplayLiveLinkSubjectReplicator;
+
+/**
+ * Interface for LiveLinkOverNDisplay module.
+ */
+class ILiveLinkOverNDisplayModule : public IModuleInterface
+{
+public:
+
+	/**
+	 * Gets a reference to the LiveLinkOverNDisplay module instance.
+	 */
+	static ILiveLinkOverNDisplayModule& Get()
+	{
+		return FModuleManager::LoadModuleChecked<ILiveLinkOverNDisplayModule>("LiveLinkOverNDisplay");
+	}
+
+	static inline bool IsAvailable()
+	{
+		return FModuleManager::Get().IsModuleLoaded("LiveLinkOverNDisplay");
+	}
+
+	virtual FNDisplayLiveLinkSubjectReplicator& GetSubjectReplicator() = 0;
+
+public:
+
+	/** Virtual destructor. */
+	virtual ~ILiveLinkOverNDisplayModule() { }
+};
+
diff --git a/Engine/Plugins/Runtime/LiveLinkOvernDisplay/Source/LiveLinkOverNDisplay/Public/LiveLinkOverNDisplaySettings.h b/Engine/Plugins/Runtime/LiveLinkOvernDisplay/Source/LiveLinkOverNDisplay/Public/LiveLinkOverNDisplaySettings.h
new file mode 100644
--- /dev/null
+++ b/Engine/Plugins/Runtime/LiveLinkOvernDisplay/Source/LiveLinkOverNDisplay/Public/LiveLinkOverNDisplaySettings.h
@@ -0,0 +1,31 @@
+// Copyright 1998-2019 Epic Games, Inc. All Rights Reserved.
+
+#pragma once
+#include "UObject/Object.h"
+#include "UObject/ObjectMacros.h"
+
+#include "LiveLinkOverNDisplaySettings.generated.h"
+
+UCLASS(config=Engine)
+class LIVELINKOVERNDISPLAY_API ULiveLinkOverNDisplaySettings : public UObject
+{
+	GENERATED_BODY()
+
+public:
+	ULiveLinkOverNDisplaySettings();
+
+	/** Whether or not LiveLink over NDisplay is enabled, from command line first or from project settings */
+	bool IsLiveLinkOverNDisplayEnabled() const { return bIsEnabledFromCommandLine.IsSet() ? bIsEnabledFromCommandLine.GetValue() : bIsEnabled; }
+
+protected:
+	/** 
+	 * Enables nDisplay specific LiveLink subjects management across nDisplay cluster. 
+	 * @note Can be overrided via the command line using -EnableLiveLinkOverNDisplay=false
+	 */
+	UPROPERTY(config, EditAnywhere, Category = "LiveLinkOverNDisplay Settings")
+	bool bIsEnabled = true;
+
+	/** Command line override value if it was set on the command line */
+	TOptional<bool> bIsEnabledFromCommandLine;
+};
+
