- Fixing CIS

diff --git a/Engine/Plugins/Runtime/LiveLinkOvernDisplay/Source/LiveLinkOverNDisplay/Private/LiveLinkOverNDisplaySettings.cpp b/Engine/Plugins/Runtime/LiveLinkOvernDisplay/Source/LiveLinkOverNDisplay/Private/LiveLinkOverNDisplaySettings.cpp
--- a/Engine/Plugins/Runtime/LiveLinkOvernDisplay/Source/LiveLinkOverNDisplay/Private/LiveLinkOverNDisplaySettings.cpp
+++ b/Engine/Plugins/Runtime/LiveLinkOvernDisplay/Source/LiveLinkOverNDisplay/Private/LiveLinkOverNDisplaySettings.cpp
@@ -3,6 +3,7 @@
 #include "LiveLinkOverNDisplaySettings.h"
 
 #include "LiveLinkOverNDisplayPrivate.h"
+#include "Misc/CommandLine.h"
 
 
 ULiveLinkOverNDisplaySettings::ULiveLinkOverNDisplaySettings()
diff --git a/Engine/Plugins/Runtime/LiveLinkOvernDisplay/Source/LiveLinkOverNDisplay/Private/NDisplayLiveLinkSubjectReplicator.cpp b/Engine/Plugins/Runtime/LiveLinkOvernDisplay/Source/LiveLinkOverNDisplay/Private/NDisplayLiveLinkSubjectReplicator.cpp
--- a/Engine/Plugins/Runtime/LiveLinkOvernDisplay/Source/LiveLinkOverNDisplay/Private/NDisplayLiveLinkSubjectReplicator.cpp
+++ b/Engine/Plugins/Runtime/LiveLinkOvernDisplay/Source/LiveLinkOverNDisplay/Private/NDisplayLiveLinkSubjectReplicator.cpp
@@ -10,6 +10,8 @@
 #include "LiveLinkSubjectSettings.h"
 #include "Misc/CoreDelegates.h"
 #include "NDisplaySlaveVirtualSubject.h"
+#include "Serialization/MemoryReader.h"
+#include "Serialization/MemoryWriter.h"
 
 
 DEFINE_LOG_CATEGORY_STATIC(LogNDisplayLiveLinkSubjectReplicator, Log, All);
diff --git a/Engine/Plugins/Runtime/LiveLinkOvernDisplay/Source/LiveLinkOverNDisplay/Private/NDisplayLiveLinkSubjectReplicator.h b/Engine/Plugins/Runtime/LiveLinkOvernDisplay/Source/LiveLinkOverNDisplay/Private/NDisplayLiveLinkSubjectReplicator.h
--- a/Engine/Plugins/Runtime/LiveLinkOvernDisplay/Source/LiveLinkOverNDisplay/Private/NDisplayLiveLinkSubjectReplicator.h
+++ b/Engine/Plugins/Runtime/LiveLinkOvernDisplay/Source/LiveLinkOverNDisplay/Private/NDisplayLiveLinkSubjectReplicator.h
@@ -5,9 +5,7 @@
 #include "Cluster/IDisplayClusterClusterSyncObject.h"
 #include "UObject/GCObject.h"
 
-
-#include "LiveLinkClientReference.h"
-
+#include "LiveLinkTypes.h"
 
 
 class UNDisplaySlaveVirtualSubject;
