- Fixing CIS again.

diff --git a/Engine/Plugins/Runtime/LiveLinkOvernDisplay/Source/LiveLinkOverNDisplay/Private/NDisplayLiveLinkSubjectReplicator.h b/Engine/Plugins/Runtime/LiveLinkOvernDisplay/Source/LiveLinkOverNDisplay/Private/NDisplayLiveLinkSubjectReplicator.h
--- a/Engine/Plugins/Runtime/LiveLinkOvernDisplay/Source/LiveLinkOverNDisplay/Private/NDisplayLiveLinkSubjectReplicator.h
+++ b/Engine/Plugins/Runtime/LiveLinkOvernDisplay/Source/LiveLinkOverNDisplay/Private/NDisplayLiveLinkSubjectReplicator.h
@@ -5,9 +5,11 @@
 #include "Cluster/IDisplayClusterClusterSyncObject.h"
 #include "UObject/GCObject.h"
 
+#include "LiveLinkRole.h"
 #include "LiveLinkTypes.h"
 
 
+class ILiveLinkClient;
 class UNDisplaySlaveVirtualSubject;
 
 
