Fix CIS. Remove unused variable in LiveLinkOverNDisplayModule.

diff --git a/Engine/Plugins/Runtime/LiveLinkOvernDisplay/Source/LiveLinkOverNDisplay/Private/LiveLinkOverNDisplayModule.h b/Engine/Plugins/Runtime/LiveLinkOvernDisplay/Source/LiveLinkOverNDisplay/Private/LiveLinkOverNDisplayModule.h
--- a/Engine/Plugins/Runtime/LiveLinkOvernDisplay/Source/LiveLinkOverNDisplay/Private/LiveLinkOverNDisplayModule.h
+++ b/Engine/Plugins/Runtime/LiveLinkOvernDisplay/Source/LiveLinkOverNDisplay/Private/LiveLinkOverNDisplayModule.h
@@ -35,9 +35,6 @@ private:
 	/** Cluster event listener delegate */
 	FDelegateHandle OnEngineLoopInitCompleteDelegate;
 
-	/** Console commands handle. */
-	TUniquePtr<FAutoConsoleCommand> StopMessageSyncCommand;
-
 	/** Replicator SyncObject used to transfer data across all cluster machines */
 	TUniquePtr<FNDisplayLiveLinkSubjectReplicator> LiveLinkReplicator;
 };
