- Fixing clang compilation

diff --git a/Engine/Plugins/Runtime/LiveLinkOvernDisplay/Source/LiveLinkOverNDisplay/Private/NDisplayLiveLinkSubjectReplicator.h b/Engine/Plugins/Runtime/LiveLinkOvernDisplay/Source/LiveLinkOverNDisplay/Private/NDisplayLiveLinkSubjectReplicator.h
--- a/Engine/Plugins/Runtime/LiveLinkOvernDisplay/Source/LiveLinkOverNDisplay/Private/NDisplayLiveLinkSubjectReplicator.h
+++ b/Engine/Plugins/Runtime/LiveLinkOvernDisplay/Source/LiveLinkOverNDisplay/Private/NDisplayLiveLinkSubjectReplicator.h
@@ -35,7 +35,6 @@ public:
 	/** Move only type since we're owning LiveLink base structs */
 	FNDisplayLiveLinkSubjectReplicator(const FNDisplayLiveLinkSubjectReplicator&) = delete;
 	FNDisplayLiveLinkSubjectReplicator& operator=(const FNDisplayLiveLinkSubjectReplicator&) = delete;
-	FNDisplayLiveLinkSubjectReplicator(FNDisplayLiveLinkSubjectReplicator&&) = default;
 
 	//~ Begin IDisplayClusterClusterSyncObject interface
 	virtual bool IsActive() const override;
@@ -83,7 +82,7 @@ private:
 	/** Cached LiveLinkClient when modular feature is registered */
 	ILiveLinkClient* LiveLinkClient = nullptr;
 
-	//CriticalSection to protect the payload. SerializeToString can bel called from other threads than Main Thread
+	//CriticalSection to protect the payload. SerializeToString can be called from other threads than Main Thread
 	mutable FCriticalSection PayloadCriticalSection;
 
 	/** SyncObject uses string to pass around data. This represents binary LiveLinkSubject information converted to string */
