Merging CL# 11895269 from Dev-Virtualproduction

- Fix for LiveLink over nDisplay crashing when default livelink preset is set in project settings

diff --git a/Engine/Plugins/Runtime/LiveLinkOvernDisplay/Source/LiveLinkOverNDisplay/Private/NDisplayLiveLinkSubjectReplicator.cpp b/Engine/Plugins/Runtime/LiveLinkOvernDisplay/Source/LiveLinkOverNDisplay/Private/NDisplayLiveLinkSubjectReplicator.cpp
--- a/Engine/Plugins/Runtime/LiveLinkOvernDisplay/Source/LiveLinkOverNDisplay/Private/NDisplayLiveLinkSubjectReplicator.cpp
+++ b/Engine/Plugins/Runtime/LiveLinkOvernDisplay/Source/LiveLinkOverNDisplay/Private/NDisplayLiveLinkSubjectReplicator.cpp
@@ -93,12 +93,10 @@ void FNDisplayLiveLinkSubjectReplicator::Activate()
 			//If we're a slave, listen for new subject and disables them if we're tracking that subject
 			FCoreDelegates::OnBeginFrame.AddRaw(this, &FNDisplayLiveLinkSubjectReplicator::OnEngineBeginFrame);
 
-			//Create our Source to hold all VirtualSubject we'll create
-			LiveLinkSourceGuid = LiveLinkClient->AddVirtualSubjectSource(TEXT("nDisplaySubjectReplicator"));
-			if (!LiveLinkSourceGuid.IsValid())
-			{
-				UE_LOG(LogNDisplayLiveLinkSubjectReplicator, Error, TEXT("Could not add VirtualSubject Source for nDisplay replication."));
-			}
+			//Used to reinitialize ourself if we are removed from livelink
+			LiveLinkClient->OnLiveLinkSourceRemoved().AddRaw(this, &FNDisplayLiveLinkSubjectReplicator::OnLiveLinkSourceRemoved);
+
+			ReInitializeVirtualSource();
 		}
 
 		//Register SyncObject on PreTick to have this behavior
@@ -362,46 +360,56 @@ void FNDisplayLiveLinkSubjectReplicator::HandleNewSubject(FArchive& Ar, FLiveLin
 		if (TrackedSubject == nullptr)
 		{
 			const FLiveLinkSubjectKey ReplicatedSubjectKey(LiveLinkSourceGuid, SubjectKey.SubjectName);
-			LiveLinkClient->AddVirtualSubject(ReplicatedSubjectKey, UNDisplaySlaveVirtualSubject::StaticClass());
-
-			//Retrieve the newly created subject to add it to our tracking array
-			const bool bIncludeDisabled = false;
-			const bool bIncludeVirtual = true;
-			for (const FLiveLinkSubjectKey& FoundSubjectKey : LiveLinkClient->GetSubjects(bIncludeDisabled, bIncludeVirtual))
+			if (LiveLinkClient->AddVirtualSubject(ReplicatedSubjectKey, UNDisplaySlaveVirtualSubject::StaticClass()))
 			{
-				if (FoundSubjectKey.SubjectName == SubjectKey.SubjectName)
+				//Retrieve the newly created subject to add it to our tracking array
+				const bool bIncludeDisabled = false;
+				const bool bIncludeVirtual = true;
+				for (const FLiveLinkSubjectKey& FoundSubjectKey : LiveLinkClient->GetSubjects(bIncludeDisabled, bIncludeVirtual))
 				{
-					//If this subject has the same subject name, verify if it's the VirtualOne. Settings object will return the VirtualSubject directly
-					if (UNDisplaySlaveVirtualSubject* NewlyCreatedSubject = Cast<UNDisplaySlaveVirtualSubject>(LiveLinkClient->GetSubjectSettings(FoundSubjectKey)))
+					if (FoundSubjectKey.SubjectName == SubjectKey.SubjectName)
 					{
-						TrackedSubject = NewlyCreatedSubject;
-						TrackedSubjects.Add(NewlyCreatedSubject);
+						//If this subject has the same subject name, verify if it's the VirtualOne. Settings object will return the VirtualSubject directly
+						if (UNDisplaySlaveVirtualSubject* NewlyCreatedSubject = Cast<UNDisplaySlaveVirtualSubject>(LiveLinkClient->GetSubjectSettings(FoundSubjectKey)))
+						{
+							TrackedSubject = NewlyCreatedSubject;
+							TrackedSubjects.Add(NewlyCreatedSubject);
+						}
 					}
 				}
+
+				checkf(TrackedSubject, TEXT("TrackedSubject '%s' was added but could not be found afterwards."), *SubjectKey.SubjectName.ToString());
+			}
+			else
+			{
+				UE_LOG(LogNDisplayLiveLinkSubjectReplicator, Warning, TEXT("Could not create new subject '%s'. It won't be tracked correctly on this node."), *SubjectKey.SubjectName.ToString());
 			}
 		}
 
-		//VirtualSubjects need translators to be able to be evaluated in different output format. Special treatment required for virtual subjects since they won't/can't exist in the client
-		if (ReplicatedSubject)
+		if (TrackedSubject)
 		{
-			const TArray<ULiveLinkFrameTranslator*>& Translators = ReplicatedSubject->GetTranslators();
-			TrackedSubject->UpdateTranslators(Translators);
-		}
-		else
-		{
-			if (ULiveLinkSubjectSettings* SubjectSetting = Cast<ULiveLinkSubjectSettings>(LiveLinkClient->GetSubjectSettings(SubjectKey)))
+			//VirtualSubjects need translators to be able to be evaluated in different output format. Special treatment required for virtual subjects since they won't/can't exist in the client
+			if (ReplicatedSubject)
 			{
-				TrackedSubject->UpdateTranslators(SubjectSetting->Translators);
+				const TArray<ULiveLinkFrameTranslator*>& Translators = ReplicatedSubject->GetTranslators();
+				TrackedSubject->UpdateTranslators(Translators);
 			}
 			else
 			{
-				UE_LOG(LogNDisplayLiveLinkSubjectReplicator, Warning, TEXT("Replicating subject '%s' but could not find its settings"), *SubjectKey.SubjectName.ToString());
+				if (ULiveLinkSubjectSettings* SubjectSetting = Cast<ULiveLinkSubjectSettings>(LiveLinkClient->GetSubjectSettings(SubjectKey)))
+				{
+					TrackedSubject->UpdateTranslators(SubjectSetting->Translators);
+				}
+				else
+				{
+					UE_LOG(LogNDisplayLiveLinkSubjectReplicator, Warning, TEXT("Replicating subject '%s' but could not find its settings"), *SubjectKey.SubjectName.ToString());
+				}
 			}
-		}
 
-		//When handling new subject, always setup role and static data
-		TrackedSubject->SetTrackedSubjectInfo(SubjectKey, SubjectRole);
-		TrackedSubject->GetStaticData().InitializeWith(SubjectFrame.StaticData);
+			//When handling new subject, always setup role and static data
+			TrackedSubject->SetTrackedSubjectInfo(SubjectKey, SubjectRole);
+			TrackedSubject->GetStaticData().InitializeWith(SubjectFrame.StaticData);
+		}
 	}
 }
 
@@ -413,9 +421,38 @@ void FNDisplayLiveLinkSubjectReplicator::ProcessLiveLinkData_Slave(EFrameType Fr
 	UNDisplaySlaveVirtualSubject* FoundTrackedSubject = *FoundTrackedSubjectPtr;
 	
 	FrameData.GetBaseData()->WorldTime = FPlatformTime::Seconds();
-	FoundTrackedSubject->UpdateFrameData(MoveTemp(FrameData));
-
+	
 	UE_LOG(LogNDisplayLiveLinkSubjectReplicator, VeryVerbose, TEXT("ProcessLiveLinkData SubjectUpdated=%s UpdateType=%d Timecode=%s")
 		, *(SubjectKey.SubjectName.ToString()), static_cast<uint8>(FrameType)
 		, *(FTimecode::FromFrameNumber(FrameData.GetBaseData()->MetaData.SceneTime.Time.FrameNumber, FrameData.GetBaseData()->MetaData.SceneTime.Rate, false).ToString()));
+	
+	FoundTrackedSubject->UpdateFrameData(MoveTemp(FrameData));
+}
+
+void FNDisplayLiveLinkSubjectReplicator::ReInitializeVirtualSource()
+{
+	//if we are rebuilding our source, let go of any subject we are currently tracking
+	for (const UNDisplaySlaveVirtualSubject* Subject : TrackedSubjects)
+	{
+		LiveLinkClient->RemoveVirtualSubject(Subject->GetSubjectKey());
+	}
+
+	//Start source with a clean slate
+	TrackedSubjects.Empty();
+
+	//Create our Source to hold all VirtualSubject we'll create
+	LiveLinkSourceGuid = LiveLinkClient->AddVirtualSubjectSource(TEXT("nDisplaySubjectReplicator"));
+	if (!LiveLinkSourceGuid.IsValid())
+	{
+		UE_LOG(LogNDisplayLiveLinkSubjectReplicator, Error, TEXT("Could not add VirtualSubject Source for nDisplay replication."));
+	}
+}
+
+void FNDisplayLiveLinkSubjectReplicator::OnLiveLinkSourceRemoved(FGuid SourceGuid)
+{
+	if (LiveLinkSourceGuid.IsValid() && LiveLinkSourceGuid == SourceGuid)
+	{
+		UE_LOG(LogNDisplayLiveLinkSubjectReplicator, Verbose, TEXT("nDisplay LiveLink Source was removed. Reinitializing ourself."));
+		ReInitializeVirtualSource();
+	}
 }
diff --git a/Engine/Plugins/Runtime/LiveLinkOvernDisplay/Source/LiveLinkOverNDisplay/Private/NDisplayLiveLinkSubjectReplicator.h b/Engine/Plugins/Runtime/LiveLinkOvernDisplay/Source/LiveLinkOverNDisplay/Private/NDisplayLiveLinkSubjectReplicator.h
--- a/Engine/Plugins/Runtime/LiveLinkOvernDisplay/Source/LiveLinkOverNDisplay/Private/NDisplayLiveLinkSubjectReplicator.h
+++ b/Engine/Plugins/Runtime/LiveLinkOvernDisplay/Source/LiveLinkOverNDisplay/Private/NDisplayLiveLinkSubjectReplicator.h
@@ -77,6 +77,12 @@ protected:
 	/** Only used on slaves, handles this frame data for a subject. If a new subject is handled, the associated VirtualSubject will be created */
 	void ProcessLiveLinkData_Slave(EFrameType FrameType, const FLiveLinkSubjectKey& SubjectKey, FLiveLinkFrameDataStruct& FrameData);
 
+	/** For slave only, remove tracked subjets, add ourself as a source  */
+	void ReInitializeVirtualSource();
+
+	/** If our source is removed (could happen if a new preset is applied), reinitialize ourself to stay awake */
+	void OnLiveLinkSourceRemoved(FGuid SourceGuid);
+
 private:
 
 	/** Cached LiveLinkClient when modular feature is registered */
