- Renamed livelink ndisplay plugin for better readability

diff --git a/Engine/Plugins/Runtime/LiveLinkOvernDisplay/LiveLinkOverNDisplay.uplugin b/Engine/Plugins/Runtime/LiveLinkOvernDisplay/LiveLinkOverNDisplay.uplugin
--- a/Engine/Plugins/Runtime/LiveLinkOvernDisplay/LiveLinkOverNDisplay.uplugin
+++ b/Engine/Plugins/Runtime/LiveLinkOvernDisplay/LiveLinkOverNDisplay.uplugin
@@ -2,7 +2,7 @@
 	"FileVersion": 3,
 	"Version": 1,
 	"VersionName": "1.0",
-	"FriendlyName": "LiveLinkOverNDisplay",
+	"FriendlyName": "Live Link Over nDisplay",
 	"CreatedBy": "Epic Games Inc",
 	"CreatedByURL": "http://epicgames.com",
 	"Description": "LiveLink subjects synchronization for nDisplay setup",
