- Fix for a crash on shutdown in LiveLink overndisplay
- Fix for livelink over ndisplay not working when map is changed

diff --git a/Engine/Plugins/Runtime/LiveLinkOvernDisplay/Source/LiveLinkOverNDisplay/Private/LiveLinkOverNDisplayModule.cpp b/Engine/Plugins/Runtime/LiveLinkOvernDisplay/Source/LiveLinkOverNDisplay/Private/LiveLinkOverNDisplayModule.cpp
--- a/Engine/Plugins/Runtime/LiveLinkOvernDisplay/Source/LiveLinkOverNDisplay/Private/LiveLinkOverNDisplayModule.cpp
+++ b/Engine/Plugins/Runtime/LiveLinkOvernDisplay/Source/LiveLinkOverNDisplay/Private/LiveLinkOverNDisplayModule.cpp
@@ -66,16 +66,37 @@ FNDisplayLiveLinkSubjectReplicator& FLiveLinkOverNDisplayModule::GetSubjectRepli
 }
 
 void FLiveLinkOverNDisplayModule::OnEngineLoopInitComplete()
+{
+	//Register to display cluster scene event delegates to know when to activate/deactivate replicator. 
+	//i.e. When loading a map, scene will be ended and all SyncObjects will be destroyed.
+	IDisplayCluster::Get().OnDisplayClusterStartScene().AddRaw(this, &FLiveLinkOverNDisplayModule::OnDisplayClusterStartSceneCallback);
+	IDisplayCluster::Get().OnDisplayClusterEndScene().AddRaw(this, &FLiveLinkOverNDisplayModule::OnDisplayClusterEndSceneCallback);
+
+	//Initialize replicator so it's ready to operate
+	LiveLinkReplicator->Initialize();
+
+	//Start scene event has already been done when engine has finished initializing so fake it
+	OnDisplayClusterStartSceneCallback();
+}
+
+void FLiveLinkOverNDisplayModule::OnDisplayClusterStartSceneCallback() const
 {
 	if (IDisplayCluster::IsAvailable() && IDisplayCluster::Get().GetOperationMode() == EDisplayClusterOperationMode::Cluster) 
 	{
 		if (GetDefault<ULiveLinkOverNDisplaySettings>()->IsLiveLinkOverNDisplayEnabled())
 		{
+			//When Scene is started, activate our replicator sync object ( registers itself )
 			LiveLinkReplicator->Activate();
 		}
 	}
 }
 
+void FLiveLinkOverNDisplayModule::OnDisplayClusterEndSceneCallback() const
+{
+	//When scene is ended, deactivate replicator sync object ( unregisters itself )
+	LiveLinkReplicator->Deactivate();
+}
+
 #undef LOCTEXT_NAMESPACE
 
 
diff --git a/Engine/Plugins/Runtime/LiveLinkOvernDisplay/Source/LiveLinkOverNDisplay/Private/LiveLinkOverNDisplayModule.h b/Engine/Plugins/Runtime/LiveLinkOvernDisplay/Source/LiveLinkOverNDisplay/Private/LiveLinkOverNDisplayModule.h
--- a/Engine/Plugins/Runtime/LiveLinkOvernDisplay/Source/LiveLinkOverNDisplay/Private/LiveLinkOverNDisplayModule.h
+++ b/Engine/Plugins/Runtime/LiveLinkOvernDisplay/Source/LiveLinkOverNDisplay/Private/LiveLinkOverNDisplayModule.h
@@ -26,9 +26,12 @@ public:
 
 	virtual FNDisplayLiveLinkSubjectReplicator& GetSubjectReplicator() override;
 
+
 private:
 
 	void OnEngineLoopInitComplete();
+	void OnDisplayClusterStartSceneCallback() const;
+	void OnDisplayClusterEndSceneCallback() const;
 
 private:
 
diff --git a/Engine/Plugins/Runtime/LiveLinkOvernDisplay/Source/LiveLinkOverNDisplay/Private/NDisplayLiveLinkSubjectReplicator.cpp b/Engine/Plugins/Runtime/LiveLinkOvernDisplay/Source/LiveLinkOverNDisplay/Private/NDisplayLiveLinkSubjectReplicator.cpp
--- a/Engine/Plugins/Runtime/LiveLinkOvernDisplay/Source/LiveLinkOverNDisplay/Private/NDisplayLiveLinkSubjectReplicator.cpp
+++ b/Engine/Plugins/Runtime/LiveLinkOvernDisplay/Source/LiveLinkOverNDisplay/Private/NDisplayLiveLinkSubjectReplicator.cpp
@@ -20,19 +20,19 @@ DEFINE_LOG_CATEGORY_STATIC(LogNDisplayLiveLinkSubjectReplicator, Log, All);
 
 FNDisplayLiveLinkSubjectReplicator::~FNDisplayLiveLinkSubjectReplicator()
 {
-	if (IDisplayCluster::IsAvailable())
-	{
-		if (IDisplayClusterClusterManager* ClusterManager = IDisplayCluster::Get().GetClusterMgr())
-		{
-			ClusterManager->UnregisterSyncObject(this);
-		}
-	}
+	Deactivate();
 	
 	if (LiveLinkClient)
 	{
 		LiveLinkClient->OnLiveLinkTicked().RemoveAll(this);
+		LiveLinkClient->OnLiveLinkSourceRemoved().RemoveAll(this);
+		LiveLinkClient->RemoveSource(LiveLinkSourceGuid);
 	}
 
+	//Cleanup our virtual source from LiveLink
+	TrackedSubjects.Empty();
+
+	IModularFeatures::Get().OnModularFeatureUnregistered().RemoveAll(this);
 	FCoreDelegates::OnFEngineLoopInitComplete.RemoveAll(this);
 	FCoreDelegates::OnBeginFrame.RemoveAll(this);
 }
@@ -77,40 +77,74 @@ void FNDisplayLiveLinkSubjectReplicator::AddReferencedObjects(FReferenceCollecto
 	}
 }
 
-void FNDisplayLiveLinkSubjectReplicator::Activate()
+void FNDisplayLiveLinkSubjectReplicator::Initialize()
 {
 	if (IModularFeatures::Get().IsModularFeatureAvailable(ILiveLinkClient::ModularFeatureName))
 	{
+		//Cache LiveLink client if it was available and listen for it being unregistered.
 		LiveLinkClient = &IModularFeatures::Get().GetModularFeature<ILiveLinkClient>(ILiveLinkClient::ModularFeatureName);
+		IModularFeatures::Get().OnModularFeatureUnregistered().AddRaw(this, &FNDisplayLiveLinkSubjectReplicator::OnModularFeatureRemoved);
 
-		if (IDisplayCluster::Get().GetClusterMgr()->IsMaster())
-		{
-			//If we're a master, listen to LiveLinkTicked callback to bundle up the SyncObject. Subjects frames will have been processed at that point.
-			LiveLinkClient->OnLiveLinkTicked().AddRaw(this, &FNDisplayLiveLinkSubjectReplicator::OnLiveLinkTicked);
-		}
-		else if (IDisplayCluster::Get().GetClusterMgr()->IsSlave())
+		if (IDisplayCluster::IsAvailable())
 		{
-			//If we're a slave, listen for new subject and disables them if we're tracking that subject
-			FCoreDelegates::OnBeginFrame.AddRaw(this, &FNDisplayLiveLinkSubjectReplicator::OnEngineBeginFrame);
+			if (IDisplayCluster::Get().GetClusterMgr()->IsMaster())
+			{
+				//If we're a master, listen to LiveLinkTicked callback to bundle up the SyncObject. Subjects frames will have been processed at that point.
+				LiveLinkClient->OnLiveLinkTicked().AddRaw(this, &FNDisplayLiveLinkSubjectReplicator::OnLiveLinkTicked);
+			}
+			else if (IDisplayCluster::Get().GetClusterMgr()->IsSlave())
+			{
+				//If we're a slave, listen for new subject and disables them if we're tracking that subject
+				FCoreDelegates::OnBeginFrame.AddRaw(this, &FNDisplayLiveLinkSubjectReplicator::OnEngineBeginFrame);
 
-			//Used to reinitialize ourself if we are removed from livelink
-			LiveLinkClient->OnLiveLinkSourceRemoved().AddRaw(this, &FNDisplayLiveLinkSubjectReplicator::OnLiveLinkSourceRemoved);
+				//Used to reinitialize ourself if we are removed from livelink
+				LiveLinkClient->OnLiveLinkSourceRemoved().AddRaw(this, &FNDisplayLiveLinkSubjectReplicator::OnLiveLinkSourceRemoved);
 
-			ReInitializeVirtualSource();
+				ReInitializeVirtualSource();
+			}
 		}
+		else
+		{
+			UE_LOG(LogNDisplayLiveLinkSubjectReplicator, Error, TEXT("Can't initialize LiveLink Subject Replicator for nDisplay because nDisplay is not available."));
+		}
+	}
+	else
+	{
+		UE_LOG(LogNDisplayLiveLinkSubjectReplicator, Error, TEXT("Can't initialize LiveLink Subject Replicator for nDisplay because LiveLink is not available."));
+	}
+}
+
+void FNDisplayLiveLinkSubjectReplicator::Activate()
+{
+	UE_LOG(LogNDisplayLiveLinkSubjectReplicator, Log, TEXT("Registering sync object."));
 
-		//Register SyncObject on PreTick to have this behavior
+	//Register SyncObject on PreTick to have this behavior
 
-		//Order of operation will be
-		//1. When LiveLink creates snapshots for the frame, Master will create the sync object payload
-		//2. On PreTick, the SyncObject will be synchronized between cluster machines
-		//3. On object processing, Slaves will push received data for each subject to have the same data then Master
+	//Order of operation will be
+	//1. When LiveLink creates snapshots for the frame, Master will create the sync object payload
+	//2. On PreTick, the SyncObject will be synchronized between cluster machines
+	//3. On object processing, Slaves will push received data for each subject to have the same data then Master
+	if (IDisplayCluster::IsAvailable())
+	{
 		IDisplayClusterClusterManager* ClusterManager = IDisplayCluster::Get().GetClusterMgr();
 		ClusterManager->RegisterSyncObject(this, EDisplayClusterSyncGroup::PreTick);
 	}
 	else
 	{
-		UE_LOG(LogNDisplayLiveLinkSubjectReplicator, Error, TEXT("Can't activate LiveLink Subject Replicator for nDisplay because LiveLink is not available."));
+		UE_LOG(LogNDisplayLiveLinkSubjectReplicator, Warning, TEXT("Can't activate LiveLink Subject Replicator for nDisplay because nDisplay is not available."));
+	}
+}
+
+void FNDisplayLiveLinkSubjectReplicator::Deactivate()
+{
+	UE_LOG(LogNDisplayLiveLinkSubjectReplicator, Log, TEXT("Unregistering sync object."));
+	
+	if (IDisplayCluster::IsAvailable())
+	{
+		if (IDisplayClusterClusterManager* ClusterManager = IDisplayCluster::Get().GetClusterMgr())
+		{
+			ClusterManager->UnregisterSyncObject(this);
+		}
 	}
 }
 
@@ -325,6 +359,7 @@ void FNDisplayLiveLinkSubjectReplicator::HandleNewSubject(FArchive& Ar, FLiveLin
 		[SubjectKey](const UNDisplaySlaveVirtualSubject* Other) { return Other->GetAssociatedSubjectKey().SubjectName == SubjectKey.SubjectName; }))
 	{
 		TrackedSubject = *FoundSujectDataPtr;
+		UE_LOG(LogNDisplayLiveLinkSubjectReplicator, Verbose, TEXT("HandleNewSubject SubjectName=%s but was already tracked. Updating its info."), *SubjectKey.SubjectName.ToString());
 	}
 
 	if (Ar.IsSaving())
@@ -402,7 +437,7 @@ void FNDisplayLiveLinkSubjectReplicator::HandleNewSubject(FArchive& Ar, FLiveLin
 				}
 				else
 				{
-					UE_LOG(LogNDisplayLiveLinkSubjectReplicator, Warning, TEXT("Replicating subject '%s' but could not find its settings"), *SubjectKey.SubjectName.ToString());
+					UE_LOG(LogNDisplayLiveLinkSubjectReplicator, Warning, TEXT("Replicating subject '%s' but could not find its settings. Translators won't work."), *SubjectKey.SubjectName.ToString());
 				}
 			}
 
@@ -456,3 +491,14 @@ void FNDisplayLiveLinkSubjectReplicator::OnLiveLinkSourceRemoved(FGuid SourceGui
 		ReInitializeVirtualSource();
 	}
 }
+
+void FNDisplayLiveLinkSubjectReplicator::OnModularFeatureRemoved(const FName& Type, IModularFeature* ModularFeature)
+{
+	//If LiveLink gets unregistered, invalidate our cached client which make the replicator do nothing
+	//To support hot reload, we would need to register for the feature to be loaded again
+	if (Type == ILiveLinkClient::ModularFeatureName)
+	{
+		UE_LOG(LogNDisplayLiveLinkSubjectReplicator, Verbose, TEXT("LiveLink feature was removed."));
+		LiveLinkClient = nullptr;
+	}
+}
diff --git a/Engine/Plugins/Runtime/LiveLinkOvernDisplay/Source/LiveLinkOverNDisplay/Private/NDisplayLiveLinkSubjectReplicator.h b/Engine/Plugins/Runtime/LiveLinkOvernDisplay/Source/LiveLinkOverNDisplay/Private/NDisplayLiveLinkSubjectReplicator.h
--- a/Engine/Plugins/Runtime/LiveLinkOvernDisplay/Source/LiveLinkOverNDisplay/Private/NDisplayLiveLinkSubjectReplicator.h
+++ b/Engine/Plugins/Runtime/LiveLinkOvernDisplay/Source/LiveLinkOverNDisplay/Private/NDisplayLiveLinkSubjectReplicator.h
@@ -10,6 +10,7 @@
 
 
 class ILiveLinkClient;
+class IModularFeature;
 class UNDisplaySlaveVirtualSubject;
 
 
@@ -53,7 +54,14 @@ public:
 	}
 	//~ End FGCObject interface
 
+	/** Initializes all required callbacks and LiveLink source */
+	void Initialize();
+	
+	/** Registers ourself as sync object in nDisplay */
 	void Activate();
+	
+	/** Unregisters ourself as sync object from nDisplay */
+	void Deactivate();
 
 protected:
 
@@ -87,6 +95,9 @@ protected:
 	/** If our source is removed (could happen if a new preset is applied), reinitialize ourself to stay awake */
 	void OnLiveLinkSourceRemoved(FGuid SourceGuid);
 
+	/** Listen for modular feature removed in case LiveLink gets unloaded */
+	void OnModularFeatureRemoved(const FName& Type , IModularFeature* ModularFeature);
+
 private:
 
 	/** Cached LiveLinkClient when modular feature is registered */
