Fix LiveLinkOvernDisplay compilation.

diff --git a/Engine/Plugins/Runtime/LiveLinkOvernDisplay/Source/LiveLinkOverNDisplay/Private/LiveLinkOverNDisplayModule.cpp b/Engine/Plugins/Runtime/LiveLinkOvernDisplay/Source/LiveLinkOverNDisplay/Private/LiveLinkOverNDisplayModule.cpp
--- a/Engine/Plugins/Runtime/LiveLinkOvernDisplay/Source/LiveLinkOverNDisplay/Private/LiveLinkOverNDisplayModule.cpp
+++ b/Engine/Plugins/Runtime/LiveLinkOvernDisplay/Source/LiveLinkOverNDisplay/Private/LiveLinkOverNDisplayModule.cpp
@@ -7,6 +7,7 @@
 #include "Features/IModularFeatures.h"
 #include "IDisplayCluster.h"
 #include "ILiveLinkClient.h"
+#include "IDisplayClusterCallbacks.h"
 #include "LiveLinkOverNDisplayPrivate.h"
 #include "LiveLinkOverNDisplaySettings.h"
 #include "Misc/CoreDelegates.h"
@@ -69,8 +70,8 @@ void FLiveLinkOverNDisplayModule::OnEngineLoopInitComplete()
 {
 	//Register to display cluster scene event delegates to know when to activate/deactivate replicator. 
 	//i.e. When loading a map, scene will be ended and all SyncObjects will be destroyed.
-	IDisplayCluster::Get().OnDisplayClusterStartScene().AddRaw(this, &FLiveLinkOverNDisplayModule::OnDisplayClusterStartSceneCallback);
-	IDisplayCluster::Get().OnDisplayClusterEndScene().AddRaw(this, &FLiveLinkOverNDisplayModule::OnDisplayClusterEndSceneCallback);
+	IDisplayCluster::Get().GetCallbacks().OnDisplayClusterStartScene().AddRaw(this, &FLiveLinkOverNDisplayModule::OnDisplayClusterStartSceneCallback);
+	IDisplayCluster::Get().GetCallbacks().OnDisplayClusterEndScene().AddRaw(this, &FLiveLinkOverNDisplayModule::OnDisplayClusterEndSceneCallback);
 
 	//Initialize replicator so it's ready to operate
 	LiveLinkReplicator->Initialize();
