Add primary node check to determine if subject being replicated comes from a source that is marked as pending kill

diff --git a/Engine/Plugins/Runtime/LiveLinkOvernDisplay/Source/LiveLinkOverNDisplay/Private/NDisplayLiveLinkSubjectReplicator.cpp b/Engine/Plugins/Runtime/LiveLinkOvernDisplay/Source/LiveLinkOverNDisplay/Private/NDisplayLiveLinkSubjectReplicator.cpp
--- a/Engine/Plugins/Runtime/LiveLinkOvernDisplay/Source/LiveLinkOverNDisplay/Private/NDisplayLiveLinkSubjectReplicator.cpp
+++ b/Engine/Plugins/Runtime/LiveLinkOvernDisplay/Source/LiveLinkOverNDisplay/Private/NDisplayLiveLinkSubjectReplicator.cpp
@@ -223,6 +223,15 @@ void FNDisplayLiveLinkSubjectReplicator::OnDataSynchronization(FArchive& Ar)
 					const FName SubjectName = SubjectKey.SubjectName;
 					TSubclassOf<ULiveLinkRole> SubjectRole = LiveLinkClient->GetSubjectRole_AnyThread(SubjectKey);
 
+					// Check if the subject's source is pending kill. If it is, the subject itself will soon be killed and we can skip evaluating it.
+					TArray<FGuid> ValidSources = LiveLinkClient->GetSources(false);
+					const bool bIsSourceValid = ValidSources.ContainsByPredicate([SubjectKey](const FGuid& Other) { return Other == SubjectKey.Source; });
+					if (!bIsSourceValid)
+					{
+						UE_LOG(LogNDisplayLiveLinkSubjectReplicator, Verbose, TEXT("Master could not evaluate Subject '%s'. Its source is marked as pending kill."), *SubjectName.ToString());
+						continue;
+					}
+
 					FLiveLinkSubjectFrameData SubjectFrameData;
 					if (LiveLinkClient->EvaluateFrame_AnyThread(SubjectName, SubjectRole, SubjectFrameData))
 					{
