PR #8375: nDisplay: Whitelisted Platform Linux for LiveLinkOverNDisplay Plugin (Contributed by [REDACTED_PERSON])

diff --git a/Engine/Plugins/Runtime/LiveLinkOvernDisplay/LiveLinkOverNDisplay.uplugin b/Engine/Plugins/Runtime/LiveLinkOvernDisplay/LiveLinkOverNDisplay.uplugin
--- a/Engine/Plugins/Runtime/LiveLinkOvernDisplay/LiveLinkOverNDisplay.uplugin
+++ b/Engine/Plugins/Runtime/LiveLinkOvernDisplay/LiveLinkOverNDisplay.uplugin
@@ -14,14 +14,15 @@
 	"CanContainContent": false,
 	"IsBetaVersion": true,
 	"Installed": false,
-	"SupportedTargetPlatforms": [ "Win64" ],
+	"SupportedTargetPlatforms": [ "Win64", "Linux" ],
 	"Modules": [
 		{
 			"Name": "LiveLinkOverNDisplay",
 			"Type": "Runtime",
 			"LoadingPhase": "Default",
 			"PlatformAllowList": [
-				"Win64"
+				"Win64",
+				"Linux"
 			]
 		}
 	],
