nDisplay: Removed plugin beta status.

diff --git a/Engine/Plugins/Runtime/LiveLinkOvernDisplay/LiveLinkOverNDisplay.uplugin b/Engine/Plugins/Runtime/LiveLinkOvernDisplay/LiveLinkOverNDisplay.uplugin
--- a/Engine/Plugins/Runtime/LiveLinkOvernDisplay/LiveLinkOverNDisplay.uplugin
+++ b/Engine/Plugins/Runtime/LiveLinkOvernDisplay/LiveLinkOverNDisplay.uplugin
@@ -12,7 +12,7 @@
 	"SupportURL": "",
 	"EnabledByDefault": false,
 	"CanContainContent": false,
-	"IsBetaVersion": true,
+	"IsBetaVersion": false,
 	"Installed": false,
 	"SupportedTargetPlatforms": [ "Win64", "Linux" ],
 	"Modules": [
diff --git a/Engine/Plugins/Runtime/nDisplay/nDisplay.uplugin b/Engine/Plugins/Runtime/nDisplay/nDisplay.uplugin
--- a/Engine/Plugins/Runtime/nDisplay/nDisplay.uplugin
+++ b/Engine/Plugins/Runtime/nDisplay/nDisplay.uplugin
@@ -12,7 +12,7 @@
 	"SupportURL": "",
 	"EnabledByDefault": false,
 	"CanContainContent": true,
-	"IsBetaVersion": true,
+	"IsBetaVersion": false,
 	"Installed": false,
 		"SupportedTargetPlatforms": [
 			"Win64",
