LLoD: Include virtual subjects in replication to secondary nodes.

diff --git a/Engine/Plugins/Runtime/LiveLinkOvernDisplay/Source/LiveLinkOverNDisplay/Private/NDisplayLiveLinkSubjectReplicator.cpp b/Engine/Plugins/Runtime/LiveLinkOvernDisplay/Source/LiveLinkOverNDisplay/Private/NDisplayLiveLinkSubjectReplicator.cpp
--- a/Engine/Plugins/Runtime/LiveLinkOvernDisplay/Source/LiveLinkOverNDisplay/Private/NDisplayLiveLinkSubjectReplicator.cpp
+++ b/Engine/Plugins/Runtime/LiveLinkOvernDisplay/Source/LiveLinkOverNDisplay/Private/NDisplayLiveLinkSubjectReplicator.cpp
@@ -224,8 +224,16 @@ void FNDisplayLiveLinkSubjectReplicator::OnDataSynchronization(FArchive& Ar)
 					TSubclassOf<ULiveLinkRole> SubjectRole = LiveLinkClient->GetSubjectRole_AnyThread(SubjectKey);
 
 					// Check if the subject's source is pending kill. If it is, the subject itself will soon be killed and we can skip evaluating it.
+
 					TArray<FGuid> ValidSources = LiveLinkClient->GetSources(false);
+
+					if (bIncludeVirtuals)
+					{
+						ValidSources += LiveLinkClient->GetVirtualSources(false);
+					}
+
 					const bool bIsSourceValid = ValidSources.ContainsByPredicate([SubjectKey](const FGuid& Other) { return Other == SubjectKey.Source; });
+
 					if (!bIsSourceValid)
 					{
 						UE_LOG(LogNDisplayLiveLinkSubjectReplicator, Verbose, TEXT("Master could not evaluate Subject '%s'. Its source is marked as pending kill."), *SubjectName.ToString());
