Fix unacceptable word usage in LiveLinkOverNDisplay.  This is change 1 of 2. The next change will rename files.

diff --git a/Engine/Plugins/Runtime/LiveLinkOvernDisplay/Config/BaseEngine.ini b/Engine/Plugins/Runtime/LiveLinkOvernDisplay/Config/BaseEngine.ini
new file mode 100644
--- /dev/null
+++ b/Engine/Plugins/Runtime/LiveLinkOvernDisplay/Config/BaseEngine.ini
@@ -0,0 +1,3 @@
+[CoreRedirects]
+; 5.1 > 5.2
++ClassRedirects=(OldName="/Script/LiveLinkOverNDisplay.NDisplaySlaveVirtualSubject",NewName="/Script/LiveLinkOverNDisplay.NDisplayAgentVirtualSubject")
diff --git a/Engine/Plugins/Runtime/LiveLinkOvernDisplay/Source/LiveLinkOverNDisplay/Private/NDisplayLiveLinkSubjectReplicator.cpp b/Engine/Plugins/Runtime/LiveLinkOvernDisplay/Source/LiveLinkOverNDisplay/Private/NDisplayLiveLinkSubjectReplicator.cpp
--- a/Engine/Plugins/Runtime/LiveLinkOvernDisplay/Source/LiveLinkOverNDisplay/Private/NDisplayLiveLinkSubjectReplicator.cpp
+++ b/Engine/Plugins/Runtime/LiveLinkOvernDisplay/Source/LiveLinkOverNDisplay/Private/NDisplayLiveLinkSubjectReplicator.cpp
@@ -53,7 +53,7 @@ FString FNDisplayLiveLinkSubjectReplicator::GetSyncId() const
 
 FString FNDisplayLiveLinkSubjectReplicator::SerializeToString() const
 {
-	//Serialization is done on the master. Payload was created on LiveLinkTicked callback
+	//Serialization is done on the controller. Payload was created on LiveLinkTicked callback
 	FScopeLock Lock(&PayloadCriticalSection);
 	const FString ThisFramePayload = LiveLinkPayload;
 	return ThisFramePayload;
@@ -61,7 +61,7 @@ FString FNDisplayLiveLinkSubjectReplicator::SerializeToString() const
 
 bool FNDisplayLiveLinkSubjectReplicator::DeserializeFromString(const FString& Str)
 {
-	//Deserialization is done on each slave
+	//Deserialization is done on each agent
 	TArray<uint8> TempBytes;
 	TempBytes.AddUninitialized(Str.Len());
 	StringToBytes(Str, TempBytes.GetData(), Str.Len());
@@ -74,7 +74,7 @@ bool FNDisplayLiveLinkSubjectReplicator::DeserializeFromString(const FString& St
 
 void FNDisplayLiveLinkSubjectReplicator::AddReferencedObjects(FReferenceCollector& Collector)
 {
-	for (const UNDisplaySlaveVirtualSubject* Subject : TrackedSubjects)
+	for (const UNDisplayAgentVirtualSubject* Subject : TrackedSubjects)
 	{
 		Collector.AddReferencedObject(Subject);
 	}
@@ -92,12 +92,12 @@ void FNDisplayLiveLinkSubjectReplicator::Initialize()
 		{
 			if (IDisplayCluster::Get().GetClusterMgr()->IsPrimary())
 			{
-				//If we're a master, listen to LiveLinkTicked callback to bundle up the SyncObject. Subjects frames will have been processed at that point.
+				//If we're a controller, listen to LiveLinkTicked callback to bundle up the SyncObject. Subjects frames will have been processed at that point.
 				LiveLinkClient->OnLiveLinkTicked().AddRaw(this, &FNDisplayLiveLinkSubjectReplicator::OnLiveLinkTicked);
 			}
 			else if (IDisplayCluster::Get().GetClusterMgr()->IsSecondary())
 			{
-				//If we're a slave, listen for new subject and disables them if we're tracking that subject
+				//If we're a agent, listen for new subject and disables them if we're tracking that subject
 				FCoreDelegates::OnBeginFrame.AddRaw(this, &FNDisplayLiveLinkSubjectReplicator::OnEngineBeginFrame);
 
 				//Used to reinitialize ourself if we are removed from livelink
@@ -124,9 +124,9 @@ void FNDisplayLiveLinkSubjectReplicator::Activate()
 	//Register SyncObject on PreTick to have this behavior
 
 	//Order of operation will be
-	//1. When LiveLink creates snapshots for the frame, Master will create the sync object payload
+	//1. When LiveLink creates snapshots for the frame, Controller will create the sync object payload
 	//2. On PreTick, the SyncObject will be synchronized between cluster machines
-	//3. On object processing, Slaves will push received data for each subject to have the same data then Master
+	//3. On object processing, Agents will push received data for each subject to have the same data then Controller
 	if (IDisplayCluster::IsAvailable())
 	{
 		IDisplayClusterClusterManager* ClusterManager = IDisplayCluster::Get().GetClusterMgr();
@@ -163,7 +163,7 @@ void FNDisplayLiveLinkSubjectReplicator::OnLiveLinkTicked()
 
 void FNDisplayLiveLinkSubjectReplicator::OnEngineBeginFrame()
 {
-	//Each frame, disable all subjects. When we are enabled, we are replicating ALL subjects from the master to each slave.
+	//Each frame, disable all subjects. When we are enabled, we are replicating ALL subjects from the controller to each agent.
 	//Only our replicated VirtualSubjects should be enabled
 	if (LiveLinkClient)
 	{
@@ -181,7 +181,7 @@ void FNDisplayLiveLinkSubjectReplicator::OnEngineBeginFrame()
 		}
 		
 		//Make sure all our replicated subjects are enabled
-		for (const UNDisplaySlaveVirtualSubject* OurSubject : TrackedSubjects)
+		for (const UNDisplayAgentVirtualSubject* OurSubject : TrackedSubjects)
 		{
 			LiveLinkClient->SetSubjectEnabled(OurSubject->GetSubjectKey(), true);
 		}
@@ -190,7 +190,7 @@ void FNDisplayLiveLinkSubjectReplicator::OnEngineBeginFrame()
 
 void FNDisplayLiveLinkSubjectReplicator::OnDataSynchronization(FArchive& Ar)
 {
-	//Master is saving, slaves are loading
+	//Controller is saving, agents are loading
 	const bool bIsSaving = Ar.IsSaving();
 	bool bContinue = bIsSaving;
 
@@ -198,7 +198,7 @@ void FNDisplayLiveLinkSubjectReplicator::OnDataSynchronization(FArchive& Ar)
 
 	//Fill in current tracked subject list. If a subject is not processed this frame, we remove it
 	TArray<FLiveLinkSubjectKey, TInlineAllocator<16>> LastFrameSubjectKeys;
-	for (const UNDisplaySlaveVirtualSubject* Subject : TrackedSubjects)
+	for (const UNDisplayAgentVirtualSubject* Subject : TrackedSubjects)
 	{
 		LastFrameSubjectKeys.Add(Subject->GetAssociatedSubjectKey());
 	}
@@ -208,7 +208,7 @@ void FNDisplayLiveLinkSubjectReplicator::OnDataSynchronization(FArchive& Ar)
 		if (LiveLinkClient)
 		{
 			/**
-			 * We're the master, and serializing data out to clients:
+			 * We're the controller, and serializing data out to clients:
 			 *	- Iterate over all of the snapshots.
 			 *		- Determine the type of frame this is (new static data, new subject, data only).
 			 *		- Serialize a True continuation byte.
@@ -217,7 +217,7 @@ void FNDisplayLiveLinkSubjectReplicator::OnDataSynchronization(FArchive& Ar)
 			 */
 			{
 				//Get all subjects except the ones that are disabled
-				//Since we are replicating using a VirtualSubject, slaves will have to remove any VirtualSubject being replicated
+				//Since we are replicating using a VirtualSubject, agents will have to remove any VirtualSubject being replicated
 				const bool bIncludeDisabled = false;
 				const bool bIncludeVirtuals = true;
 				TArray<FLiveLinkSubjectKey> ThisFramesSubjects = LiveLinkClient->GetSubjects(bIncludeDisabled, bIncludeVirtuals);
@@ -239,17 +239,17 @@ void FNDisplayLiveLinkSubjectReplicator::OnDataSynchronization(FArchive& Ar)
 
 					if (!bIsSourceValid)
 					{
-						UE_LOG(LogNDisplayLiveLinkSubjectReplicator, Verbose, TEXT("Master could not evaluate Subject '%s'. Its source is marked as pending kill."), *SubjectName.ToString());
+						UE_LOG(LogNDisplayLiveLinkSubjectReplicator, Verbose, TEXT("Controller could not evaluate Subject '%s'. Its source is marked as pending kill."), *SubjectName.ToString());
 						continue;
 					}
 
 					FLiveLinkSubjectFrameData SubjectFrameData;
 					if (LiveLinkClient->EvaluateFrame_AnyThread(SubjectName, SubjectRole, SubjectFrameData))
 					{
-						if (UNDisplaySlaveVirtualSubject** FoundSujectDataPtr = TrackedSubjects.FindByPredicate(
-							[SubjectKey](const UNDisplaySlaveVirtualSubject* Other) { return Other->GetAssociatedSubjectKey().SubjectName == SubjectKey.SubjectName; }))
+						if (UNDisplayAgentVirtualSubject** FoundSujectDataPtr = TrackedSubjects.FindByPredicate(
+							[SubjectKey](const UNDisplayAgentVirtualSubject* Other) { return Other->GetAssociatedSubjectKey().SubjectName == SubjectKey.SubjectName; }))
 						{
-							UNDisplaySlaveVirtualSubject* FoundSubjectData = *FoundSujectDataPtr;
+							UNDisplayAgentVirtualSubject* FoundSubjectData = *FoundSujectDataPtr;
 							if (SubjectFrameData.StaticData != FoundSubjectData->GetStaticData() 
 								|| FoundSubjectData->GetAssociatedSubjectKey() != SubjectKey 
 								|| FoundSubjectData->GetRole() != SubjectRole)
@@ -275,7 +275,7 @@ void FNDisplayLiveLinkSubjectReplicator::OnDataSynchronization(FArchive& Ar)
 					}
 					else
 					{
-						UE_LOG(LogNDisplayLiveLinkSubjectReplicator, Verbose, TEXT("Master could not evaluate Subject '%s'. It could have starved."), *SubjectName.ToString());
+						UE_LOG(LogNDisplayLiveLinkSubjectReplicator, Verbose, TEXT("Controller could not evaluate Subject '%s'. It could have starved."), *SubjectName.ToString());
 					}
 				}
 			}
@@ -288,7 +288,7 @@ void FNDisplayLiveLinkSubjectReplicator::OnDataSynchronization(FArchive& Ar)
 	else
 	{
 		/**
-		 * We're a slave, and serializing data from the master:
+		 * We're a agent, and serializing data from the controller:
 		 *	- Serialize the continuation byte
 		 *		- If False, we're done.
 		 *		- If True:
@@ -310,7 +310,7 @@ void FNDisplayLiveLinkSubjectReplicator::OnDataSynchronization(FArchive& Ar)
 				HandleFrame(Ar, FrameType, SubjectKey, SubjectRole, SubjectFrameData);
 				if (!SubjectKey.SubjectName.IsNone())
 				{
-					ProcessLiveLinkData_Slave(FrameType, SubjectKey, SubjectFrameData.FrameData);
+					ProcessLiveLinkData_Agent(FrameType, SubjectKey, SubjectFrameData.FrameData);
 
 					LastFrameSubjectKeys.RemoveSingleSwap(SubjectKey);
 				}
@@ -322,7 +322,7 @@ void FNDisplayLiveLinkSubjectReplicator::OnDataSynchronization(FArchive& Ar)
 		}
 	}
 
-	// Cleanup missing subjects. Slaves will need to remove VirtualSubjects that were created
+	// Cleanup missing subjects. Agents will need to remove VirtualSubjects that were created
 	if (LiveLinkClient)
 	{
 		for (const FLiveLinkSubjectKey& ItemSubjectKey : LastFrameSubjectKeys)
@@ -335,7 +335,7 @@ void FNDisplayLiveLinkSubjectReplicator::OnDataSynchronization(FArchive& Ar)
 				LiveLinkClient->RemoveVirtualSubject(ReplicatedSubjectKey);
 			}
 
-			const int32 FoundIndex = TrackedSubjects.IndexOfByPredicate([ItemSubjectKey](const UNDisplaySlaveVirtualSubject* Other) { return Other->GetAssociatedSubjectKey() == ItemSubjectKey; });
+			const int32 FoundIndex = TrackedSubjects.IndexOfByPredicate([ItemSubjectKey](const UNDisplayAgentVirtualSubject* Other) { return Other->GetAssociatedSubjectKey() == ItemSubjectKey; });
 			if (TrackedSubjects.IsValidIndex(FoundIndex))
 			{
 				TrackedSubjects.RemoveAtSwap(FoundIndex);
@@ -374,9 +374,9 @@ void FNDisplayLiveLinkSubjectReplicator::HandleNewSubject(FArchive& Ar, FLiveLin
 	UE_LOG(LogNDisplayLiveLinkSubjectReplicator, Verbose, TEXT("HandleNewSubject SubjectName=%s"), *SubjectKey.SubjectName.ToString());
 
 	//Verify if we're already tracking that subject
-	UNDisplaySlaveVirtualSubject* TrackedSubject = nullptr;
-	if (UNDisplaySlaveVirtualSubject** FoundSujectDataPtr = TrackedSubjects.FindByPredicate(
-		[SubjectKey](const UNDisplaySlaveVirtualSubject* Other) { return Other->GetAssociatedSubjectKey().SubjectName == SubjectKey.SubjectName; }))
+	UNDisplayAgentVirtualSubject* TrackedSubject = nullptr;
+	if (UNDisplayAgentVirtualSubject** FoundSujectDataPtr = TrackedSubjects.FindByPredicate(
+		[SubjectKey](const UNDisplayAgentVirtualSubject* Other) { return Other->GetAssociatedSubjectKey().SubjectName == SubjectKey.SubjectName; }))
 	{
 		TrackedSubject = *FoundSujectDataPtr;
 		UE_LOG(LogNDisplayLiveLinkSubjectReplicator, Verbose, TEXT("HandleNewSubject SubjectName=%s but was already tracked. Updating its info."), *SubjectKey.SubjectName.ToString());
@@ -390,7 +390,7 @@ void FNDisplayLiveLinkSubjectReplicator::HandleNewSubject(FArchive& Ar, FLiveLin
 		//If we're not tracking that subject yet, create it
 		if (TrackedSubject == nullptr)
 		{
-			TrackedSubject = NewObject<UNDisplaySlaveVirtualSubject>(GetTransientPackage());
+			TrackedSubject = NewObject<UNDisplayAgentVirtualSubject>(GetTransientPackage());
 			TrackedSubjects.Add(TrackedSubject);
 		}
 		
@@ -405,7 +405,7 @@ void FNDisplayLiveLinkSubjectReplicator::HandleNewSubject(FArchive& Ar, FLiveLin
 		SubjectRole = StaticLoadClass(UObject::StaticClass(), NULL, *RoleClassName, NULL, LOAD_None, NULL);
 		check(SubjectRole);
 
-		//On Slaves, create the virtual subject directly in LiveLink so it's part of the system
+		//On Agents, create the virtual subject directly in LiveLink so it's part of the system
 		ULiveLinkVirtualSubject* ReplicatedSubject = Cast<ULiveLinkVirtualSubject>(LiveLinkClient->GetSubjectSettings(SubjectKey));
 
 		//Always disable the subject we're replicating
@@ -415,7 +415,7 @@ void FNDisplayLiveLinkSubjectReplicator::HandleNewSubject(FArchive& Ar, FLiveLin
 		if (TrackedSubject == nullptr)
 		{
 			const FLiveLinkSubjectKey ReplicatedSubjectKey(LiveLinkSourceGuid, SubjectKey.SubjectName);
-			if (LiveLinkClient->AddVirtualSubject(ReplicatedSubjectKey, UNDisplaySlaveVirtualSubject::StaticClass()))
+			if (LiveLinkClient->AddVirtualSubject(ReplicatedSubjectKey, UNDisplayAgentVirtualSubject::StaticClass()))
 			{
 				//Retrieve the newly created subject to add it to our tracking array
 				const bool bIncludeDisabled = false;
@@ -425,7 +425,7 @@ void FNDisplayLiveLinkSubjectReplicator::HandleNewSubject(FArchive& Ar, FLiveLin
 					if (FoundSubjectKey.SubjectName == SubjectKey.SubjectName)
 					{
 						//If this subject has the same subject name, verify if it's the VirtualOne. Settings object will return the VirtualSubject directly
-						if (UNDisplaySlaveVirtualSubject* NewlyCreatedSubject = Cast<UNDisplaySlaveVirtualSubject>(LiveLinkClient->GetSubjectSettings(FoundSubjectKey)))
+						if (UNDisplayAgentVirtualSubject* NewlyCreatedSubject = Cast<UNDisplayAgentVirtualSubject>(LiveLinkClient->GetSubjectSettings(FoundSubjectKey)))
 						{
 							TrackedSubject = NewlyCreatedSubject;
 							TrackedSubjects.Add(NewlyCreatedSubject);
@@ -468,12 +468,12 @@ void FNDisplayLiveLinkSubjectReplicator::HandleNewSubject(FArchive& Ar, FLiveLin
 	}
 }
 
-void FNDisplayLiveLinkSubjectReplicator::ProcessLiveLinkData_Slave(EFrameType FrameType, const FLiveLinkSubjectKey& SubjectKey, FLiveLinkFrameDataStruct& FrameData)
+void FNDisplayLiveLinkSubjectReplicator::ProcessLiveLinkData_Agent(EFrameType FrameType, const FLiveLinkSubjectKey& SubjectKey, FLiveLinkFrameDataStruct& FrameData)
 {
 	//Find associated tracked subject data to broadcast role and static data. Should always be found
-	UNDisplaySlaveVirtualSubject** FoundTrackedSubjectPtr = TrackedSubjects.FindByPredicate([SubjectKey](const UNDisplaySlaveVirtualSubject* Other) { return Other->GetAssociatedSubjectKey() == SubjectKey; });
+	UNDisplayAgentVirtualSubject** FoundTrackedSubjectPtr = TrackedSubjects.FindByPredicate([SubjectKey](const UNDisplayAgentVirtualSubject* Other) { return Other->GetAssociatedSubjectKey() == SubjectKey; });
 	check(FoundTrackedSubjectPtr);
-	UNDisplaySlaveVirtualSubject* FoundTrackedSubject = *FoundTrackedSubjectPtr;
+	UNDisplayAgentVirtualSubject* FoundTrackedSubject = *FoundTrackedSubjectPtr;
 	
 	FrameData.GetBaseData()->WorldTime = FPlatformTime::Seconds();
 	
@@ -487,7 +487,7 @@ void FNDisplayLiveLinkSubjectReplicator::ProcessLiveLinkData_Slave(EFrameType Fr
 void FNDisplayLiveLinkSubjectReplicator::ReInitializeVirtualSource()
 {
 	//if we are rebuilding our source, let go of any subject we are currently tracking
-	for (const UNDisplaySlaveVirtualSubject* Subject : TrackedSubjects)
+	for (const UNDisplayAgentVirtualSubject* Subject : TrackedSubjects)
 	{
 		LiveLinkClient->RemoveVirtualSubject(Subject->GetSubjectKey());
 	}
diff --git a/Engine/Plugins/Runtime/LiveLinkOvernDisplay/Source/LiveLinkOverNDisplay/Private/NDisplayLiveLinkSubjectReplicator.h b/Engine/Plugins/Runtime/LiveLinkOvernDisplay/Source/LiveLinkOverNDisplay/Private/NDisplayLiveLinkSubjectReplicator.h
--- a/Engine/Plugins/Runtime/LiveLinkOvernDisplay/Source/LiveLinkOverNDisplay/Private/NDisplayLiveLinkSubjectReplicator.h
+++ b/Engine/Plugins/Runtime/LiveLinkOvernDisplay/Source/LiveLinkOverNDisplay/Private/NDisplayLiveLinkSubjectReplicator.h
@@ -13,12 +13,12 @@ template <typename T> class TSubclassOf;
 
 class ILiveLinkClient;
 class IModularFeature;
-class UNDisplaySlaveVirtualSubject;
+class UNDisplayAgentVirtualSubject;
 
 
 /**
- * Classes used to replicates data to be used for each frame for each enabled subjects on the Master machines.
- * Slaves will use that replicated data to use the same thing on each machines
+ * Classes used to replicates data to be used for each frame for each enabled subjects on the Controller machines.
+ * Agents will use that replicated data to use the same thing on each machines
  */
 class LIVELINKOVERNDISPLAY_API FNDisplayLiveLinkSubjectReplicator : public IDisplayClusterClusterSyncObject, public FGCObject
 {
@@ -69,13 +69,13 @@ protected:
 
 	/**
      * Callback to be notified when LiveLinkClient was ticked so SyncObject can be serialized
-     * @note Only called on Master
+     * @note Only called on Controller
      */
 	void OnLiveLinkTicked();
 
 	/**
 	 * Hook to beginning of frames to make sure our replicated virtual subjects are the ones enabled
-	 * @note Only called on Slave
+	 * @note Only called on Agent
 	 */
 	void OnEngineBeginFrame();
 
@@ -88,10 +88,10 @@ protected:
 	/** Process a frame for a specific subject */
 	void HandleFrame(FArchive& Ar, EFrameType& FrameType, FLiveLinkSubjectKey& SubjectKey, TSubclassOf<ULiveLinkRole>& SubjectRole, FLiveLinkSubjectFrameData& SubjectFrame);
 
-	/** Only used on slaves, handles this frame data for a subject. If a new subject is handled, the associated VirtualSubject will be created */
-	void ProcessLiveLinkData_Slave(EFrameType FrameType, const FLiveLinkSubjectKey& SubjectKey, FLiveLinkFrameDataStruct& FrameData);
+	/** Only used on agents, handles this frame data for a subject. If a new subject is handled, the associated VirtualSubject will be created */
+	void ProcessLiveLinkData_Agent(EFrameType FrameType, const FLiveLinkSubjectKey& SubjectKey, FLiveLinkFrameDataStruct& FrameData);
 
-	/** For slave only, remove tracked subjets, add ourself as a source  */
+	/** For agent only, remove tracked subjets, add ourself as a source  */
 	void ReInitializeVirtualSource();
 
 	/** If our source is removed (could happen if a new preset is applied), reinitialize ourself to stay awake */
@@ -111,8 +111,8 @@ private:
 	/** SyncObject uses string to pass around data. This represents binary LiveLinkSubject information converted to string */
 	FString LiveLinkPayload;
 
-	/** List of Subjects that we are replicating across cluster. On Slaves, it will actually be added to the LiveLink subject's list */
-	TArray<UNDisplaySlaveVirtualSubject*> TrackedSubjects;
+	/** List of Subjects that we are replicating across cluster. On Agents, it will actually be added to the LiveLink subject's list */
+	TArray<UNDisplayAgentVirtualSubject*> TrackedSubjects;
 
 	/** Guid associated to our Virtual Subject Source */
 	FGuid LiveLinkSourceGuid;
diff --git a/Engine/Plugins/Runtime/LiveLinkOvernDisplay/Source/LiveLinkOverNDisplay/Private/NDisplaySlaveVirtualSubject.cpp b/Engine/Plugins/Runtime/LiveLinkOvernDisplay/Source/LiveLinkOverNDisplay/Private/NDisplaySlaveVirtualSubject.cpp
--- a/Engine/Plugins/Runtime/LiveLinkOvernDisplay/Source/LiveLinkOverNDisplay/Private/NDisplaySlaveVirtualSubject.cpp
+++ b/Engine/Plugins/Runtime/LiveLinkOvernDisplay/Source/LiveLinkOverNDisplay/Private/NDisplaySlaveVirtualSubject.cpp
@@ -5,23 +5,23 @@
 
 
 
-void UNDisplaySlaveVirtualSubject::Update()
+void UNDisplayAgentVirtualSubject::Update()
 {
 	UpdateTranslatorsForThisFrame();
 }
 
-void UNDisplaySlaveVirtualSubject::UpdateFrameData(FLiveLinkFrameDataStruct&& NewFrameData)
+void UNDisplayAgentVirtualSubject::UpdateFrameData(FLiveLinkFrameDataStruct&& NewFrameData)
 {
 	UpdateFrameDataSnapshot(MoveTemp(NewFrameData));
 }
 
-void UNDisplaySlaveVirtualSubject::SetTrackedSubjectInfo(const FLiveLinkSubjectKey& InSubjectKey, TSubclassOf<ULiveLinkRole> InRole)
+void UNDisplayAgentVirtualSubject::SetTrackedSubjectInfo(const FLiveLinkSubjectKey& InSubjectKey, TSubclassOf<ULiveLinkRole> InRole)
 {
 	AssociatedSubject = InSubjectKey;
 	Role = InRole;
 }
 
-void UNDisplaySlaveVirtualSubject::UpdateTranslators(const TArray<ULiveLinkFrameTranslator*>& SourceTranslators)
+void UNDisplayAgentVirtualSubject::UpdateTranslators(const TArray<ULiveLinkFrameTranslator*>& SourceTranslators)
 {
 	for (const ULiveLinkFrameTranslator* SourceTranslator : SourceTranslators)
 	{
diff --git a/Engine/Plugins/Runtime/LiveLinkOvernDisplay/Source/LiveLinkOverNDisplay/Private/NDisplaySlaveVirtualSubject.h b/Engine/Plugins/Runtime/LiveLinkOvernDisplay/Source/LiveLinkOverNDisplay/Private/NDisplaySlaveVirtualSubject.h
--- a/Engine/Plugins/Runtime/LiveLinkOvernDisplay/Source/LiveLinkOverNDisplay/Private/NDisplaySlaveVirtualSubject.h
+++ b/Engine/Plugins/Runtime/LiveLinkOvernDisplay/Source/LiveLinkOverNDisplay/Private/NDisplaySlaveVirtualSubject.h
@@ -11,12 +11,12 @@ class ULiveLinkFrameTranslator;
 
 
 /**
- * LiveLink VirtualSubject used on nDisplay slaves to replicate real subjects
- * Master sends data to use for each frame to make sure all machines display / uses the same data
+ * LiveLink VirtualSubject used on nDisplay agents to replicate real subjects
+ * Controller sends data to use for each frame to make sure all machines display / uses the same data
  * Uses Replicator object that this module creates.
  */
 UCLASS()
-class UNDisplaySlaveVirtualSubject : public ULiveLinkVirtualSubject
+class UNDisplayAgentVirtualSubject : public ULiveLinkVirtualSubject
 {
 	GENERATED_BODY()
 
