Fix unacceptable word usage. This is change 2 of 2 for addressing live link over nDisplay.

diff --git a/Engine/Plugins/Runtime/LiveLinkOvernDisplay/Source/LiveLinkOverNDisplay/Private/NDisplayAgentVirtualSubject.cpp b/Engine/Plugins/Runtime/LiveLinkOvernDisplay/Source/LiveLinkOverNDisplay/Private/NDisplayAgentVirtualSubject.cpp
new file mode 100644
--- /dev/null
+++ b/Engine/Plugins/Runtime/LiveLinkOvernDisplay/Source/LiveLinkOverNDisplay/Private/NDisplayAgentVirtualSubject.cpp
@@ -0,0 +1,31 @@
+// Copyright Epic Games, Inc. All Rights Reserved.
+
+#include "NDisplayAgentVirtualSubject.h"
+
+
+
+
+void UNDisplayAgentVirtualSubject::Update()
+{
+	UpdateTranslatorsForThisFrame();
+}
+
+void UNDisplayAgentVirtualSubject::UpdateFrameData(FLiveLinkFrameDataStruct&& NewFrameData)
+{
+	UpdateFrameDataSnapshot(MoveTemp(NewFrameData));
+}
+
+void UNDisplayAgentVirtualSubject::SetTrackedSubjectInfo(const FLiveLinkSubjectKey& InSubjectKey, TSubclassOf<ULiveLinkRole> InRole)
+{
+	AssociatedSubject = InSubjectKey;
+	Role = InRole;
+}
+
+void UNDisplayAgentVirtualSubject::UpdateTranslators(const TArray<ULiveLinkFrameTranslator*>& SourceTranslators)
+{
+	for (const ULiveLinkFrameTranslator* SourceTranslator : SourceTranslators)
+	{
+		ULiveLinkFrameTranslator* NewTranslator = DuplicateObject(SourceTranslator, this);
+		FrameTranslators.Add(NewTranslator);
+	}
+}
diff --git a/Engine/Plugins/Runtime/LiveLinkOvernDisplay/Source/LiveLinkOverNDisplay/Private/NDisplayAgentVirtualSubject.h b/Engine/Plugins/Runtime/LiveLinkOvernDisplay/Source/LiveLinkOverNDisplay/Private/NDisplayAgentVirtualSubject.h
new file mode 100644
--- /dev/null
+++ b/Engine/Plugins/Runtime/LiveLinkOvernDisplay/Source/LiveLinkOverNDisplay/Private/NDisplayAgentVirtualSubject.h
@@ -0,0 +1,42 @@
+// Copyright Epic Games, Inc. All Rights Reserved.
+
+#pragma once
+
+#include "LiveLinkVirtualSubject.h"
+
+
+#include "NDisplayAgentVirtualSubject.generated.h"
+
+class ULiveLinkFrameTranslator;
+
+
+/**
+ * LiveLink VirtualSubject used on nDisplay agents to replicate real subjects
+ * Controller sends data to use for each frame to make sure all machines display / uses the same data
+ * Uses Replicator object that this module creates.
+ */
+UCLASS()
+class UNDisplayAgentVirtualSubject : public ULiveLinkVirtualSubject
+{
+	GENERATED_BODY()
+
+public:
+	
+	//~ Begin ULiveLinkVirtualSubject interface
+	virtual void Update() override;
+	//~ End ULiveLinkVirtualSubject interface
+
+	void UpdateFrameData(FLiveLinkFrameDataStruct&& NewFrameData);
+
+
+	void SetTrackedSubjectInfo(const FLiveLinkSubjectKey& SubjectKey, TSubclassOf<ULiveLinkRole> Role);
+
+	void UpdateTranslators(const TArray<ULiveLinkFrameTranslator*>& SourceTranslators);
+
+	FLiveLinkSubjectKey GetAssociatedSubjectKey() const { return AssociatedSubject; }
+
+private:
+
+	/** Real subject that this virtual subject replicates */
+	FLiveLinkSubjectKey AssociatedSubject;
+};
diff --git a/Engine/Plugins/Runtime/LiveLinkOvernDisplay/Source/LiveLinkOverNDisplay/Private/NDisplayLiveLinkSubjectReplicator.cpp b/Engine/Plugins/Runtime/LiveLinkOvernDisplay/Source/LiveLinkOverNDisplay/Private/NDisplayLiveLinkSubjectReplicator.cpp
--- a/Engine/Plugins/Runtime/LiveLinkOvernDisplay/Source/LiveLinkOverNDisplay/Private/NDisplayLiveLinkSubjectReplicator.cpp
+++ b/Engine/Plugins/Runtime/LiveLinkOvernDisplay/Source/LiveLinkOverNDisplay/Private/NDisplayLiveLinkSubjectReplicator.cpp
@@ -11,7 +11,7 @@
 #include "LiveLinkRole.h"
 #include "LiveLinkSubjectSettings.h"
 #include "Misc/CoreDelegates.h"
-#include "NDisplaySlaveVirtualSubject.h"
+#include "NDisplayAgentVirtualSubject.h"
 #include "Serialization/MemoryReader.h"
 #include "Serialization/MemoryWriter.h"
 #include "UObject/Package.h"
diff --git a/Engine/Plugins/Runtime/LiveLinkOvernDisplay/Source/LiveLinkOverNDisplay/Private/NDisplaySlaveVirtualSubject.cpp b/Engine/Plugins/Runtime/LiveLinkOvernDisplay/Source/LiveLinkOverNDisplay/Private/NDisplaySlaveVirtualSubject.cpp
deleted file mode 100644
diff --git a/Engine/Plugins/Runtime/LiveLinkOvernDisplay/Source/LiveLinkOverNDisplay/Private/NDisplaySlaveVirtualSubject.h b/Engine/Plugins/Runtime/LiveLinkOvernDisplay/Source/LiveLinkOverNDisplay/Private/NDisplaySlaveVirtualSubject.h
deleted file mode 100644
