[LiveLinkOvernDisplay] A quick fix to let LL work
- A failover-friendly version of the fix will arrive a bit later

#rnx

diff --git a/Engine/Plugins/Runtime/LiveLinkOvernDisplay/Source/LiveLinkOverNDisplay/Private/NDisplayLiveLinkSubjectReplicator.cpp b/Engine/Plugins/Runtime/LiveLinkOvernDisplay/Source/LiveLinkOverNDisplay/Private/NDisplayLiveLinkSubjectReplicator.cpp
--- a/Engine/Plugins/Runtime/LiveLinkOvernDisplay/Source/LiveLinkOverNDisplay/Private/NDisplayLiveLinkSubjectReplicator.cpp
+++ b/Engine/Plugins/Runtime/LiveLinkOvernDisplay/Source/LiveLinkOverNDisplay/Private/NDisplayLiveLinkSubjectReplicator.cpp
@@ -61,6 +61,15 @@ FString FNDisplayLiveLinkSubjectReplicator::SerializeToString() const
 
 bool FNDisplayLiveLinkSubjectReplicator::DeserializeFromString(const FString& Str)
 {
+	// New nDisplay Failover subsystem calls DeserializeFromString on every node,
+	// including the primary one. This breaks the original logic of this replicator.
+	// We expect this method to be called on non-primary nodes only.
+	// This is a temporary. A failover-friendly fix will arrive later.
+	if (IDisplayCluster::Get().GetClusterMgr()->IsPrimary())
+	{
+		return true;
+	}
+
 	//Deserialization is done on each agent
 	TArray<uint8> TempBytes;
 	TempBytes.AddUninitialized(Str.Len());
