[nDisplay] Fixed LivelinkOverNDisplay replication logic in order to support new failover

- Fixed cluster crash when P-node fails with failover enabled
- Implemented post-failover re-initialization to keep LiveLink data replication going

[RN] Fixed LiveLink replication crashes when primary node fails with failover enabled

diff --git a/Engine/Plugins/Runtime/LiveLinkOvernDisplay/Source/LiveLinkOverNDisplay/Private/NDisplayLiveLinkSubjectReplicator.cpp b/Engine/Plugins/Runtime/LiveLinkOvernDisplay/Source/LiveLinkOverNDisplay/Private/NDisplayLiveLinkSubjectReplicator.cpp
--- a/Engine/Plugins/Runtime/LiveLinkOvernDisplay/Source/LiveLinkOverNDisplay/Private/NDisplayLiveLinkSubjectReplicator.cpp
+++ b/Engine/Plugins/Runtime/LiveLinkOvernDisplay/Source/LiveLinkOverNDisplay/Private/NDisplayLiveLinkSubjectReplicator.cpp
@@ -7,6 +7,7 @@
 #include "DisplayClusterEnums.h"
 #include "Features/IModularFeatures.h"
 #include "IDisplayCluster.h"
+#include "IDisplayClusterCallbacks.h"
 #include "ILiveLinkClient.h"
 #include "LiveLinkRole.h"
 #include "LiveLinkSubjectSettings.h"
@@ -20,24 +21,21 @@
 DEFINE_LOG_CATEGORY_STATIC(LogNDisplayLiveLinkSubjectReplicator, Log, All);
 
 
+namespace UE::LiveLinkOverNDisplay::Private
+{
+	// This variable was introduced as part of a hotfix. To avoid any class layout changes,
+	// it's defined here, not as an FNDisplayLiveLinkSubjectReplicator class member.
+	// 
+	// It's used to pass re-initialization request from a failover AnyThread callback
+	// to the main thread. Once the request is processed, the flag will be reset.
+	static std::atomic<bool> bBecamePrimary = false;
+}
+
 
 FNDisplayLiveLinkSubjectReplicator::~FNDisplayLiveLinkSubjectReplicator()
 {
 	Deactivate();
-	
-	if (LiveLinkClient)
-	{
-		LiveLinkClient->OnLiveLinkTicked().RemoveAll(this);
-		LiveLinkClient->OnLiveLinkSourceRemoved().RemoveAll(this);
-		LiveLinkClient->RemoveSource(LiveLinkSourceGuid);
-	}
-
-	//Cleanup our virtual source from LiveLink
-	TrackedSubjects.Empty();
-
-	IModularFeatures::Get().OnModularFeatureUnregistered().RemoveAll(this);
-	FCoreDelegates::OnFEngineLoopInitComplete.RemoveAll(this);
-	FCoreDelegates::OnBeginFrame.RemoveAll(this);
+	Release();
 }
 
 bool FNDisplayLiveLinkSubjectReplicator::IsActive() const
@@ -61,10 +59,8 @@ FString FNDisplayLiveLinkSubjectReplicator::SerializeToString() const
 
 bool FNDisplayLiveLinkSubjectReplicator::DeserializeFromString(const FString& Str)
 {
-	// New nDisplay Failover subsystem calls DeserializeFromString on every node,
-	// including the primary one. This breaks the original logic of this replicator.
-	// We expect this method to be called on non-primary nodes only.
-	// This is a temporary. A failover-friendly fix will arrive later.
+	// On the primary node (controller), we should skip this. The underlying logic
+	// is not supposed to be called on the primary nodes (controllers).
 	if (IDisplayCluster::Get().GetClusterMgr()->IsPrimary())
 	{
 		return true;
@@ -106,6 +102,8 @@ void FNDisplayLiveLinkSubjectReplicator::Initialize()
 			}
 			else if (IDisplayCluster::Get().GetClusterMgr()->IsSecondary())
 			{
+				IDisplayCluster::Get().GetCallbacks().OnDisplayClusterFailoverPrimaryNodeChanged().AddRaw(this, &FNDisplayLiveLinkSubjectReplicator::OnDisplayClusterPrimaryNodeChanged);
+
 				//If we're a agent, listen for new subject and disables them if we're tracking that subject
 				FCoreDelegates::OnBeginFrame.AddRaw(this, &FNDisplayLiveLinkSubjectReplicator::OnEngineBeginFrame);
 
@@ -126,6 +124,27 @@ void FNDisplayLiveLinkSubjectReplicator::Initialize()
 	}
 }
 
+void FNDisplayLiveLinkSubjectReplicator::Release()
+{
+	IModularFeatures::Get().OnModularFeatureUnregistered().RemoveAll(this);
+	FCoreDelegates::OnBeginFrame.RemoveAll(this);
+
+	if (LiveLinkClient && IModularFeatures::Get().IsModularFeatureAvailable(ILiveLinkClient::ModularFeatureName))
+	{
+		LiveLinkClient->OnLiveLinkTicked().RemoveAll(this);
+		LiveLinkClient->OnLiveLinkSourceRemoved().RemoveAll(this);
+		LiveLinkClient->RemoveSource(LiveLinkSourceGuid);
+	}
+
+	if (IDisplayCluster::IsAvailable())
+	{
+		IDisplayCluster::Get().GetCallbacks().OnDisplayClusterFailoverPrimaryNodeChanged().RemoveAll(this);
+	}
+
+	//Cleanup our virtual source from LiveLink
+	TrackedSubjects.Empty();
+}
+
 void FNDisplayLiveLinkSubjectReplicator::Activate()
 {
 	UE_LOG(LogNDisplayLiveLinkSubjectReplicator, Log, TEXT("Registering sync object."));
@@ -172,23 +191,30 @@ void FNDisplayLiveLinkSubjectReplicator::OnLiveLinkTicked()
 
 void FNDisplayLiveLinkSubjectReplicator::OnEngineBeginFrame()
 {
+	if (UE::LiveLinkOverNDisplay::Private::bBecamePrimary)
+	{
+		// It's not an agent anymore. Now it's a controller, therefore re-initialize.
+		Release();
+		Initialize();
+
+		// Since we're primary now, we need to enable all LL subjects. They were intentionally
+		// disabled when we this node was an agent (non-primary).
+		constexpr bool bIncludeDisabled = true; // including disabled ones
+		constexpr bool bIncludeVirtual = true;
+		SetSubjectsEnabled(bIncludeDisabled, bIncludeVirtual, true);
+
+		UE::LiveLinkOverNDisplay::Private::bBecamePrimary = false;
+	}
 	//Each frame, disable all subjects. When we are enabled, we are replicating ALL subjects from the controller to each agent.
 	//Only our replicated VirtualSubjects should be enabled
-	if (LiveLinkClient)
+	else if (LiveLinkClient)
 	{
 		check(IDisplayCluster::Get().GetClusterMgr()->IsSecondary());
 
 		const bool bIncludeDisabled = false;
 		const bool bIncludeVirtual = true;
-		const TArray<FLiveLinkSubjectKey> Subjects = LiveLinkClient->GetSubjects(bIncludeDisabled, bIncludeVirtual);
-		for (const FLiveLinkSubjectKey& Subject : Subjects)
-		{
-			if (Subject.Source != LiveLinkSourceGuid)
-			{
-				LiveLinkClient->SetSubjectEnabled(Subject, false);
-			}
-		}
-		
+		SetSubjectsEnabled(bIncludeDisabled, bIncludeVirtual, false);
+
 		//Make sure all our replicated subjects are enabled
 		for (const UNDisplayAgentVirtualSubject* OurSubject : TrackedSubjects)
 		{
@@ -531,3 +557,24 @@ void FNDisplayLiveLinkSubjectReplicator::OnModularFeatureRemoved(const FName& Ty
 		LiveLinkClient = nullptr;
 	}
 }
+
+void FNDisplayLiveLinkSubjectReplicator::OnDisplayClusterPrimaryNodeChanged(const FString& NewPrimaryId)
+{
+	const FString ThisNodeId = IDisplayCluster::Get().GetClusterMgr()->GetNodeId();
+	if (ThisNodeId.Equals(NewPrimaryId, ESearchCase::IgnoreCase))
+	{
+		UE::LiveLinkOverNDisplay::Private::bBecamePrimary = true;
+	}
+}
+
+void FNDisplayLiveLinkSubjectReplicator::SetSubjectsEnabled(bool bIncludeDisabled, bool bIncludeVirtual, bool bEnabled)
+{
+	const TArray<FLiveLinkSubjectKey> Subjects = LiveLinkClient->GetSubjects(bIncludeDisabled, bIncludeVirtual);
+	for (const FLiveLinkSubjectKey& Subject : Subjects)
+	{
+		if (Subject.Source != LiveLinkSourceGuid)
+		{
+			LiveLinkClient->SetSubjectEnabled(Subject, bEnabled);
+		}
+	}
+}
diff --git a/Engine/Plugins/Runtime/LiveLinkOvernDisplay/Source/LiveLinkOverNDisplay/Private/NDisplayLiveLinkSubjectReplicator.h b/Engine/Plugins/Runtime/LiveLinkOvernDisplay/Source/LiveLinkOverNDisplay/Private/NDisplayLiveLinkSubjectReplicator.h
--- a/Engine/Plugins/Runtime/LiveLinkOvernDisplay/Source/LiveLinkOverNDisplay/Private/NDisplayLiveLinkSubjectReplicator.h
+++ b/Engine/Plugins/Runtime/LiveLinkOvernDisplay/Source/LiveLinkOverNDisplay/Private/NDisplayLiveLinkSubjectReplicator.h
@@ -59,6 +59,9 @@ public:
 	/** Initializes all required callbacks and LiveLink source */
 	void Initialize();
 	
+	/** Releases all internal resources and callbacks */
+	void Release();
+
 	/** Registers ourself as sync object in nDisplay */
 	void Activate();
 	
@@ -100,6 +103,12 @@ protected:
 	/** Listen for modular feature removed in case LiveLink gets unloaded */
 	void OnModularFeatureRemoved(const FName& Type , IModularFeature* ModularFeature);
 
+	/** Post-failover handler. It's a callback that informs about new p-node ID */
+	void OnDisplayClusterPrimaryNodeChanged(const FString& NewPrimaryId);
+
+	/** A helper to enable/disable all the subjects */
+	void SetSubjectsEnabled(bool bIncludeDisabled, bool bIncludeVirtual, bool bEnabled);
+
 private:
 
 	/** Cached LiveLinkClient when modular feature is registered */
