[nDisplay] Re-shaping a 5.7 hotfix
- The initial fix had to follow the hotfix rules. This re-shaping just makes a variable to be encapsulated within class.

#rnx

#Synced-CL 0

diff --git a/Engine/Plugins/Runtime/LiveLinkOvernDisplay/Source/LiveLinkOverNDisplay/Private/NDisplayLiveLinkSubjectReplicator.cpp b/Engine/Plugins/Runtime/LiveLinkOvernDisplay/Source/LiveLinkOverNDisplay/Private/NDisplayLiveLinkSubjectReplicator.cpp
--- a/Engine/Plugins/Runtime/LiveLinkOvernDisplay/Source/LiveLinkOverNDisplay/Private/NDisplayLiveLinkSubjectReplicator.cpp
+++ b/Engine/Plugins/Runtime/LiveLinkOvernDisplay/Source/LiveLinkOverNDisplay/Private/NDisplayLiveLinkSubjectReplicator.cpp
@@ -21,17 +21,6 @@
 DEFINE_LOG_CATEGORY_STATIC(LogNDisplayLiveLinkSubjectReplicator, Log, All);
 
 
-namespace UE::LiveLinkOverNDisplay::Private
-{
-	// This variable was introduced as part of a hotfix. To avoid any class layout changes,
-	// it's defined here, not as an FNDisplayLiveLinkSubjectReplicator class member.
-	// 
-	// It's used to pass re-initialization request from a failover AnyThread callback
-	// to the main thread. Once the request is processed, the flag will be reset.
-	static std::atomic<bool> bBecamePrimary = false;
-}
-
-
 FNDisplayLiveLinkSubjectReplicator::~FNDisplayLiveLinkSubjectReplicator()
 {
 	Deactivate();
@@ -191,7 +180,7 @@ void FNDisplayLiveLinkSubjectReplicator::OnLiveLinkTicked()
 
 void FNDisplayLiveLinkSubjectReplicator::OnEngineBeginFrame()
 {
-	if (UE::LiveLinkOverNDisplay::Private::bBecamePrimary)
+	if (bBecamePrimary)
 	{
 		// It's not an agent anymore. Now it's a controller, therefore re-initialize.
 		Release();
@@ -203,7 +192,7 @@ void FNDisplayLiveLinkSubjectReplicator::OnEngineBeginFrame()
 		constexpr bool bIncludeVirtual = true;
 		SetSubjectsEnabled(bIncludeDisabled, bIncludeVirtual, true);
 
-		UE::LiveLinkOverNDisplay::Private::bBecamePrimary = false;
+		bBecamePrimary = false;
 	}
 	//Each frame, disable all subjects. When we are enabled, we are replicating ALL subjects from the controller to each agent.
 	//Only our replicated VirtualSubjects should be enabled
@@ -563,7 +552,7 @@ void FNDisplayLiveLinkSubjectReplicator::OnDisplayClusterPrimaryNodeChanged(cons
 	const FString ThisNodeId = IDisplayCluster::Get().GetClusterMgr()->GetNodeId();
 	if (ThisNodeId.Equals(NewPrimaryId, ESearchCase::IgnoreCase))
 	{
-		UE::LiveLinkOverNDisplay::Private::bBecamePrimary = true;
+		bBecamePrimary = true;
 	}
 }
 
diff --git a/Engine/Plugins/Runtime/LiveLinkOvernDisplay/Source/LiveLinkOverNDisplay/Private/NDisplayLiveLinkSubjectReplicator.h b/Engine/Plugins/Runtime/LiveLinkOvernDisplay/Source/LiveLinkOverNDisplay/Private/NDisplayLiveLinkSubjectReplicator.h
--- a/Engine/Plugins/Runtime/LiveLinkOvernDisplay/Source/LiveLinkOverNDisplay/Private/NDisplayLiveLinkSubjectReplicator.h
+++ b/Engine/Plugins/Runtime/LiveLinkOvernDisplay/Source/LiveLinkOverNDisplay/Private/NDisplayLiveLinkSubjectReplicator.h
@@ -7,6 +7,8 @@
 
 #include "LiveLinkTypes.h"
 
+#include <atomic>
+
 class ULiveLinkRole;
 template <typename T> class TSubclassOf;
 
@@ -125,4 +127,7 @@ private:
 
 	/** Guid associated to our Virtual Subject Source */
 	FGuid LiveLinkSourceGuid;
+
+	/** Used to re-initialize current node after failure handling */
+	std::atomic<bool> bBecamePrimary = false;
 };
